Mirror 4092349 from Dev-Framework
Speculative fix for shutdown crash, avoid use of LoadModule in DataValidator ShutdownModule and delete delegates registered in another module when BlueprintGraph shuts down

diff --git a/Engine/Plugins/Editor/DataValidation/Source/DataValidation/Private/DataValidationModule.cpp b/Engine/Plugins/Editor/DataValidation/Source/DataValidation/Private/DataValidationModule.cpp
--- a/Engine/Plugins/Editor/DataValidation/Source/DataValidation/Private/DataValidationModule.cpp
+++ b/Engine/Plugins/Editor/DataValidation/Source/DataValidation/Private/DataValidationModule.cpp
@@ -90,7 +90,7 @@ void FDataValidationModule::ShutdownModule()
 {
 	if (!IsRunningCommandlet() && !IsRunningGame())
 	{
-		FContentBrowserModule* ContentBrowserModule = (FContentBrowserModule*)(FModuleManager::Get().GetModule(TEXT("ContentBrowser")));
+		FContentBrowserModule* ContentBrowserModule = FModuleManager::GetModulePtr<FContentBrowserModule>(TEXT("ContentBrowser"));
 		if (ContentBrowserModule)
 		{
 			TArray<FContentBrowserMenuExtender_SelectedAssets>& CBMenuExtenderDelegates = ContentBrowserModule->GetAllAssetViewContextMenuExtenders();
@@ -99,8 +99,11 @@ void FDataValidationModule::ShutdownModule()
 		}
 
 		// remove menu extension
-		FLevelEditorModule& LevelEditorModule = FModuleManager::LoadModuleChecked<FLevelEditorModule>("LevelEditor");
-		LevelEditorModule.GetMenuExtensibilityManager()->RemoveExtender(MenuExtender);
+		FLevelEditorModule* LevelEditorModule = FModuleManager::GetModulePtr<FLevelEditorModule>("LevelEditor");
+		if (LevelEditorModule)
+		{
+			LevelEditorModule->GetMenuExtensibilityManager()->RemoveExtender(MenuExtender);
+		}
 		MenuExtender = nullptr;
 	}
 }
diff --git a/Engine/Source/Editor/BlueprintGraph/Classes/EdGraphSchema_K2.h b/Engine/Source/Editor/BlueprintGraph/Classes/EdGraphSchema_K2.h
--- a/Engine/Source/Editor/BlueprintGraph/Classes/EdGraphSchema_K2.h
+++ b/Engine/Source/Editor/BlueprintGraph/Classes/EdGraphSchema_K2.h
@@ -1094,6 +1094,8 @@ public:
 	/** Create sub menu that shows all possible variables that can be used to replace the existing variable reference */
 	static void GetReplaceVariableMenu(class FMenuBuilder& MenuBuilder, class UK2Node_Variable* Variable, UBlueprint* OwnerBlueprint, bool bReplaceExistingVariable = false);
 
+	/** Function called when the owning module is shut down */ 
+	static void Shutdown();
 private:
 
 	/**
diff --git a/Engine/Source/Editor/BlueprintGraph/Private/BlueprintGraphModule.cpp b/Engine/Source/Editor/BlueprintGraph/Private/BlueprintGraphModule.cpp
--- a/Engine/Source/Editor/BlueprintGraph/Private/BlueprintGraphModule.cpp
+++ b/Engine/Source/Editor/BlueprintGraph/Private/BlueprintGraphModule.cpp
@@ -1,8 +1,14 @@
 // Copyright 1998-2018 Epic Games, Inc. All Rights Reserved.
 
 #include "BlueprintGraphModule.h"
-#include "Modules/ModuleManager.h"
 
+#include "EdGraphSchema_K2.h"
+#include "Modules/ModuleManager.h"
 
 IMPLEMENT_MODULE( FBlueprintGraphModule, BlueprintGraph );
 
+void FBlueprintGraphModule::ShutdownModule()
+{
+	UEdGraphSchema_K2::Shutdown();
+}
+
diff --git a/Engine/Source/Editor/BlueprintGraph/Private/EdGraphSchema_K2.cpp b/Engine/Source/Editor/BlueprintGraph/Private/EdGraphSchema_K2.cpp
--- a/Engine/Source/Editor/BlueprintGraph/Private/EdGraphSchema_K2.cpp
+++ b/Engine/Source/Editor/BlueprintGraph/Private/EdGraphSchema_K2.cpp
@@ -2450,6 +2450,12 @@ public:
 		return *AutocastFunctionMap;
 	}
 
+	static void Shutdown()
+	{
+		delete AutocastFunctionMap;
+		AutocastFunctionMap = nullptr;
+	}
+
 	static void OnProjectHotReloaded(bool bWasTriggeredAutomatically)
 	{
 		if (AutocastFunctionMap)
@@ -2489,6 +2495,12 @@ public:
 
 FAutocastFunctionMap* FAutocastFunctionMap::AutocastFunctionMap = nullptr;
 
+void UEdGraphSchema_K2::Shutdown()
+{
+	FAutocastFunctionMap::Shutdown();
+}
+
+
 bool UEdGraphSchema_K2::SearchForAutocastFunction(const UEdGraphPin* OutputPin, const UEdGraphPin* InputPin, /*out*/ FName& TargetFunction, /*out*/ UClass*& FunctionOwner) const
 {
 	// NOTE: Under no circumstances should anyone *ever* add a questionable cast to this function.
diff --git a/Engine/Source/Editor/BlueprintGraph/Public/BlueprintGraphModule.h b/Engine/Source/Editor/BlueprintGraph/Public/BlueprintGraphModule.h
--- a/Engine/Source/Editor/BlueprintGraph/Public/BlueprintGraphModule.h
+++ b/Engine/Source/Editor/BlueprintGraph/Public/BlueprintGraphModule.h
@@ -25,6 +25,7 @@ public:
 	 */
 	BLUEPRINTGRAPH_API TArray<FActionMenuRejectionTest>& GetExtendedActionMenuFilters() { return ExtendedMenuFilters; }
 
+	virtual void ShutdownModule() override;
 private:
 	TArray<FActionMenuRejectionTest> ExtendedMenuFilters;
 };
