- Fixed invalid errors when generating a Landscape LOD with Create Package Per Asset enabled after generating with it disabled
- Caused by [REDACTED_PERSON] not in sync (race condition between Directory Watcher and Validation Manager)

diff --git a/Engine/Plugins/Editor/DataValidation/Source/DataValidation/Private/DataValidationManager.cpp b/Engine/Plugins/Editor/DataValidation/Source/DataValidation/Private/DataValidationManager.cpp
--- a/Engine/Plugins/Editor/DataValidation/Source/DataValidation/Private/DataValidationManager.cpp
+++ b/Engine/Plugins/Editor/DataValidation/Source/DataValidation/Private/DataValidationManager.cpp
@@ -217,12 +217,33 @@ bool UDataValidationManager::IsPathExcludedFromValidation(const FString& Path) c
 void UDataValidationManager::ValidateAllSavedPackages()
 {
 	FAssetRegistryModule& AssetRegistryModule = FModuleManager::LoadModuleChecked<FAssetRegistryModule>("AssetRegistry");
-	TArray<FAssetData> Assets;
+	IAssetRegistry& AssetRegistry = AssetRegistryModule.Get();
 
+	// Prior to validation, make sure Asset Registry is updated.
+	// DirectoryWatcher is responsible of scanning modified asset files, but validation can be called before.
+	if (SavedPackagesToValidate.Num())
+	{
+		TArray<FString> FilesToScan;
+		FilesToScan.Reserve(SavedPackagesToValidate.Num());
+		for (FName PackageName : SavedPackagesToValidate)
+		{
+			FString PackageFilename;
+			if (FPackageName::FindPackageFileWithoutExtension(FPackageName::LongPackageNameToFilename(PackageName.ToString()), PackageFilename))
+			{
+				FilesToScan.Add(PackageFilename);
+			}
+		}
+		if (FilesToScan.Num())
+		{
+			AssetRegistry.ScanModifiedAssetFiles(FilesToScan);
+		}
+	}
+
+	TArray<FAssetData> Assets;
 	for (FName PackageName : SavedPackagesToValidate)
 	{
 		// We need to query the in-memory data as the disk cache may not be accurate
-		AssetRegistryModule.Get().GetAssetsByPackageName(PackageName, Assets);
+		AssetRegistry.GetAssetsByPackageName(PackageName, Assets);
 	}
 
 	ValidateOnSave(Assets);
