Adding scriptable extension points to the Data Validation system. Python scripts should register themselves with the EditorValidatorSubsystem, while Blueprint and C++ validators can be found automatically. These validators do not require a custom class, but instead take an individual asset in as an argument to run tests on.
#rnx

diff --git a/Engine/Plugins/Editor/DataValidation/Source/DataValidation/DataValidation.Build.cs b/Engine/Plugins/Editor/DataValidation/Source/DataValidation/DataValidation.Build.cs
--- a/Engine/Plugins/Editor/DataValidation/Source/DataValidation/DataValidation.Build.cs
+++ b/Engine/Plugins/Editor/DataValidation/Source/DataValidation/DataValidation.Build.cs
@@ -11,7 +11,8 @@ public class DataValidation : ModuleRules
 				"Core",
 				"CoreUObject",
 				"Engine",
-				"TargetPlatform"
+				"TargetPlatform",
+				"EditorSubsystem"
 			}
 		);
 //		goto through these one by one and remove extra ones
@@ -32,8 +33,9 @@ public class DataValidation : ModuleRules
 				"GraphEditor",
 				"BlueprintGraph",
 				"KismetCompiler",
-				"SandboxFile"
-			}
+				"SandboxFile",
+                "Blutility"
+            }
 		);
 	}
 }
diff --git a/Engine/Plugins/Editor/DataValidation/Source/DataValidation/Private/DataValidationCommandlet.cpp b/Engine/Plugins/Editor/DataValidation/Source/DataValidation/Private/DataValidationCommandlet.cpp
--- a/Engine/Plugins/Editor/DataValidation/Source/DataValidation/Private/DataValidationCommandlet.cpp
+++ b/Engine/Plugins/Editor/DataValidation/Source/DataValidation/Private/DataValidationCommandlet.cpp
@@ -1,11 +1,11 @@
 // Copyright 1998-2019 Epic Games, Inc. All Rights Reserved.
 
 #include "DataValidationCommandlet.h"
-#include "DataValidationManager.h"
-
 #include "IAssetRegistry.h"
 #include "AssetRegistryHelpers.h"
 #include "AssetRegistryModule.h"
+#include "Editor.h"
+#include "EditorValidatorSubsystem.h"
 
 
 DEFINE_LOG_CATEGORY_STATIC(LogDataValidation, Warning, All);
@@ -40,10 +40,10 @@ bool UDataValidationCommandlet::ValidateData()
 	TArray<FAssetData> AssetDataList;
 	AssetRegistryModule.Get().GetAllAssets(AssetDataList);
 
-	UDataValidationManager* DataValidationManager = UDataValidationManager::Get();
-	check(DataValidationManager);
+	UEditorValidatorSubsystem* EditorValidationSubsystem = GEditor->GetEditorSubsystem<UEditorValidatorSubsystem>();
+	check(EditorValidationSubsystem);
 
-	DataValidationManager->ValidateAssets(AssetDataList);
+	EditorValidationSubsystem->ValidateAssets(AssetDataList);
 
 	return true;
 }
diff --git a/Engine/Plugins/Editor/DataValidation/Source/DataValidation/Private/DataValidationManager.cpp b/Engine/Plugins/Editor/DataValidation/Source/DataValidation/Private/DataValidationManager.cpp
--- a/Engine/Plugins/Editor/DataValidation/Source/DataValidation/Private/DataValidationManager.cpp
+++ b/Engine/Plugins/Editor/DataValidation/Source/DataValidation/Private/DataValidationManager.cpp
@@ -13,29 +13,29 @@
 
 #define LOCTEXT_NAMESPACE "DataValidationManager"
 
-UDataValidationManager* GDataValidationManager = nullptr;
+UDEPRECATED_DataValidationManager* GDataValidationManager = nullptr;
 
 /**
- * UDataValidationManager
+ * UDEPRECATED_DataValidationManager
  */
 
-UDataValidationManager::UDataValidationManager(const FObjectInitializer& ObjectInitializer)
+UDEPRECATED_DataValidationManager::UDEPRECATED_DataValidationManager(const FObjectInitializer& ObjectInitializer)
 	: Super(ObjectInitializer)
 {
 	DataValidationManagerClassName = FSoftClassPath(TEXT("/Script/DataValidation.DataValidationManager"));
 	bValidateOnSave = true;
 }
 
-UDataValidationManager* UDataValidationManager::Get()
+UDEPRECATED_DataValidationManager* UDEPRECATED_DataValidationManager::Get()
 {
 	if (GDataValidationManager == nullptr)
 	{
-		FSoftClassPath DataValidationManagerClassName = (UDataValidationManager::StaticClass()->GetDefaultObject<UDataValidationManager>())->DataValidationManagerClassName;
+		FSoftClassPath DataValidationManagerClassName = (UDEPRECATED_DataValidationManager::StaticClass()->GetDefaultObject<UDEPRECATED_DataValidationManager>())->DataValidationManagerClassName;
 
 		UClass* SingletonClass = DataValidationManagerClassName.TryLoadClass<UObject>();
 		checkf(SingletonClass != nullptr, TEXT("Data validation config value DataValidationManagerClassName is not a valid class name."));
 
-		GDataValidationManager = NewObject<UDataValidationManager>(GetTransientPackage(), SingletonClass, NAME_None);
+		GDataValidationManager = NewObject<UDEPRECATED_DataValidationManager>(GetTransientPackage(), SingletonClass, NAME_None);
 		checkf(GDataValidationManager != nullptr, TEXT("Data validation config value DataValidationManagerClassName is not a subclass of UDataValidationManager."))
 
 		GDataValidationManager->AddToRoot();
@@ -45,7 +45,7 @@ UDataValidationManager* UDataValidationManager::Get()
 	return GDataValidationManager;
 }
 
-void UDataValidationManager::Initialize()
+void UDEPRECATED_DataValidationManager::Initialize()
 {
 	FMessageLogInitializationOptions InitOptions;
 	InitOptions.bShowFilters = true;
@@ -54,11 +54,11 @@ void UDataValidationManager::Initialize()
 	MessageLogModule.RegisterLogListing("DataValidation", LOCTEXT("DataValidation", "Data Validation"), InitOptions);
 }
 
-UDataValidationManager::~UDataValidationManager()
+UDEPRECATED_DataValidationManager::~UDEPRECATED_DataValidationManager()
 {
 }
 
-EDataValidationResult UDataValidationManager::IsObjectValid(UObject* InObject, TArray<FText>& ValidationErrors) const
+EDataValidationResult UDEPRECATED_DataValidationManager::IsObjectValid(UObject* InObject, TArray<FText>& ValidationErrors) const
 {
 	if (ensure(InObject))
 	{
@@ -68,7 +68,7 @@ EDataValidationResult UDataValidationManager::IsObjectValid(UObject* InObject, T
 	return EDataValidationResult::NotValidated;
 }
 
-EDataValidationResult UDataValidationManager::IsAssetValid(FAssetData& AssetData, TArray<FText>& ValidationErrors) const
+EDataValidationResult UDEPRECATED_DataValidationManager::IsAssetValid(FAssetData& AssetData, TArray<FText>& ValidationErrors) const
 {
 	if (AssetData.IsValid())
 	{
@@ -82,7 +82,7 @@ EDataValidationResult UDataValidationManager::IsAssetValid(FAssetData& AssetData
 	return EDataValidationResult::Invalid;
 }
 
-int32 UDataValidationManager::ValidateAssets(TArray<FAssetData> AssetDataList, bool bSkipExcludedDirectories, bool bShowIfNoFailures) const
+int32 UDEPRECATED_DataValidationManager::ValidateAssets(TArray<FAssetData> AssetDataList, bool bSkipExcludedDirectories, bool bShowIfNoFailures) const
 {
 	FScopedSlowTask SlowTask(1.0f, LOCTEXT("ValidatingDataTask", "Validating Data..."));
 	SlowTask.Visibility = bShowIfNoFailures ? ESlowTaskVisibility::ForceVisible : ESlowTaskVisibility::Invisible;
@@ -170,7 +170,7 @@ int32 UDataValidationManager::ValidateAssets(TArray<FAssetData> AssetDataList, b
 	return NumInvalidFiles;
 }
 
-void UDataValidationManager::ValidateOnSave(TArray<FAssetData> AssetDataList) const
+void UDEPRECATED_DataValidationManager::ValidateOnSave(TArray<FAssetData> AssetDataList) const
 {
 	// Only validate if enabled and not auto saving
 	if (!bValidateOnSave || GEditor->IsAutosaving())
@@ -188,7 +188,7 @@ void UDataValidationManager::ValidateOnSave(TArray<FAssetData> AssetDataList) co
 	}
 }
 
-void UDataValidationManager::ValidateSavedPackage(FName PackageName)
+void UDEPRECATED_DataValidationManager::ValidateSavedPackage(FName PackageName)
 {
 	// Only validate if enabled and not auto saving
 	if (!bValidateOnSave || GEditor->IsAutosaving())
@@ -198,10 +198,10 @@ void UDataValidationManager::ValidateSavedPackage(FName PackageName)
 
 	SavedPackagesToValidate.AddUnique(PackageName);
 
-	GEditor->GetTimerManager()->SetTimerForNextTick(this, &UDataValidationManager::ValidateAllSavedPackages);
+	GEditor->GetTimerManager()->SetTimerForNextTick(this, &UDEPRECATED_DataValidationManager::ValidateAllSavedPackages);
 }
 
-bool UDataValidationManager::IsPathExcludedFromValidation(const FString& Path) const
+bool UDEPRECATED_DataValidationManager::IsPathExcludedFromValidation(const FString& Path) const
 {
 	for (const FDirectoryPath& ExcludedPath : ExcludedDirectories)
 	{
@@ -214,7 +214,7 @@ bool UDataValidationManager::IsPathExcludedFromValidation(const FString& Path) c
 	return false;
 }
 
-void UDataValidationManager::ValidateAllSavedPackages()
+void UDEPRECATED_DataValidationManager::ValidateAllSavedPackages()
 {
 	FAssetRegistryModule& AssetRegistryModule = FModuleManager::LoadModuleChecked<FAssetRegistryModule>("AssetRegistry");
 	IAssetRegistry& AssetRegistry = AssetRegistryModule.Get();
diff --git a/Engine/Plugins/Editor/DataValidation/Source/DataValidation/Private/DataValidationModule.cpp b/Engine/Plugins/Editor/DataValidation/Source/DataValidation/Private/DataValidationModule.cpp
--- a/Engine/Plugins/Editor/DataValidation/Source/DataValidation/Private/DataValidationModule.cpp
+++ b/Engine/Plugins/Editor/DataValidation/Source/DataValidation/Private/DataValidationModule.cpp
@@ -1,7 +1,6 @@
 // Copyright 1998-2019 Epic Games, Inc. All Rights Reserved.
 
 #include "DataValidationModule.h"
-#include "DataValidationManager.h"
 #include "UObject/Object.h"
 #include "UObject/SoftObjectPath.h"
 #include "GameFramework/HUD.h"
@@ -23,6 +22,7 @@
 
 #include "WorkspaceMenuStructure.h"
 #include "WorkspaceMenuStructureModule.h"
+#include "EditorValidatorSubsystem.h"
 
 #define LOCTEXT_NAMESPACE "DataValidationModule"
 
@@ -205,10 +205,10 @@ void FDataValidationModule::ValidateAssets(TArray<FAssetData> SelectedAssets, bo
 		SelectedAssets = DependentAssets.Array();
 	}
 
-	UDataValidationManager* DataValidationManager = UDataValidationManager::Get();
-	if (DataValidationManager)
+	UEditorValidatorSubsystem* EditorValidationSubsystem = GEditor->GetEditorSubsystem<UEditorValidatorSubsystem>();
+	if (EditorValidationSubsystem)
 	{
-		DataValidationManager->ValidateAssets(SelectedAssets, false);
+		EditorValidationSubsystem->ValidateAssets(SelectedAssets, false);
 	}
 }
 
@@ -292,10 +292,10 @@ void FDataValidationModule::CreateDataValidationContentBrowserPathMenu(FMenuBuil
 
 void FDataValidationModule::OnPackageSaved(const FString& PackageFileName, UObject* PackageObj)
 {
-	UDataValidationManager* DataValidationManager = UDataValidationManager::Get();
-	if (DataValidationManager && PackageObj)
+	UEditorValidatorSubsystem* EditorValidationSubsystem = GEditor->GetEditorSubsystem<UEditorValidatorSubsystem>();
+	if (EditorValidationSubsystem && PackageObj)
 	{
-		DataValidationManager->ValidateSavedPackage(PackageObj->GetFName());
+		EditorValidationSubsystem->ValidateSavedPackage(PackageObj->GetFName());
 	}
 }
 
diff --git a/Engine/Plugins/Editor/DataValidation/Source/DataValidation/Private/EditorValidatorBase.cpp b/Engine/Plugins/Editor/DataValidation/Source/DataValidation/Private/EditorValidatorBase.cpp
new file mode 100644
--- /dev/null
+++ b/Engine/Plugins/Editor/DataValidation/Source/DataValidation/Private/EditorValidatorBase.cpp
@@ -0,0 +1,91 @@
+// Copyright 1998-2019 Epic Games, Inc. All Rights Reserved.
+
+#include "EditorValidatorBase.h"
+
+#include "Editor.h"
+#include "AssetToolsModule.h"
+#include "ObjectTools.h"
+#include "Dialogs/Dialogs.h"
+#include "EditorValidatorSubsystem.h"
+#include "Logging/MessageLog.h"
+#include "Misc/UObjectToken.h"
+#include "Logging/TokenizedMessage.h"
+
+DEFINE_LOG_CATEGORY_STATIC(LogContentValidation, Log, Log);
+#define LOCTEXT_NAMESPACE "AssetValidation"
+
+UEditorValidatorBase::UEditorValidatorBase()
+	: Super()
+{
+	bIsEnabled = true;
+}
+
+bool UEditorValidatorBase::CanValidateAsset_Implementation(UObject* InAsset) const
+{
+	return false;
+}
+
+EDataValidationResult UEditorValidatorBase::ValidateLoadedAsset_Implementation(UObject* InAsset, TArray<FText>& ValidationErrors)
+{
+	return EDataValidationResult::NotValidated;
+}
+
+void UEditorValidatorBase::AssetFails(UObject* InAsset, const FText& InMessage, TArray<FText>& ValidationErrors)
+{
+	FFormatNamedArguments Arguments;
+	if (InAsset)
+	{
+		Arguments.Add(TEXT("AssetName"), FText::FromName(InAsset->GetFName()));	
+	}
+	Arguments.Add(TEXT("CustomMessage"), InMessage);
+	Arguments.Add(TEXT("ValidatorName"), FText::FromString(GetClass()->GetName()));
+
+	FText FailureMessage = FText::Format(LOCTEXT("AssetCheck_Message_Display", "{AssetName} failed: {CustomMessage}. ({ValidatorName})"), Arguments);
+
+	if(LogContentValidation.GetVerbosity() >= ELogVerbosity::Verbose)
+	{
+		LogElapsedTime(Arguments);
+
+	}
+
+	ValidationErrors.Add(FailureMessage);
+	ValidationResult = EDataValidationResult::Invalid;
+}
+
+void UEditorValidatorBase::LogElapsedTime(FFormatNamedArguments &Arguments)
+{
+	FDateTime CurrentTime = FDateTime::Now();
+	FTimespan ElapsedTimeSpan = (CurrentTime - ValidationTime);
+	float ElapsedTimeMS = ElapsedTimeSpan.GetTotalMilliseconds();
+	FNumberFormattingOptions TimeFormat;
+	TimeFormat.MinimumFractionalDigits = 5;
+	Arguments.Add(TEXT("ElapsedTime"), FText::AsNumber(ElapsedTimeMS, &TimeFormat));
+	FText ElapsedTimeMessage = FText::Format(LOCTEXT("ElapsedTime", "Checking {AssetName} with {ValidatorName} took {ElapsedTime} ms."), Arguments);
+	UE_LOG(LogContentValidation, Verbose, TEXT("%s"), *ElapsedTimeMessage.ToString());
+}
+
+void UEditorValidatorBase::AssetPasses(UObject* InAsset)
+{
+
+	if (LogContentValidation.GetVerbosity() >= ELogVerbosity::Verbose)
+	{
+		FFormatNamedArguments Arguments;
+		if (InAsset)
+		{
+			Arguments.Add(TEXT("AssetName"), FText::FromName(InAsset->GetFName()));
+		}
+		Arguments.Add(TEXT("ValidatorName"), FText::FromString(GetClass()->GetName()));
+
+		LogElapsedTime(Arguments);
+	}
+
+	ValidationResult = EDataValidationResult::Valid;
+}
+
+void UEditorValidatorBase::ResetValidationState()
+{
+	ValidationResult = EDataValidationResult::NotValidated;
+	ValidationTime = FDateTime::Now();
+}
+
+#undef LOCTEXT_NAMESPACE
\ No newline at end of file
diff --git a/Engine/Plugins/Editor/DataValidation/Source/DataValidation/Private/EditorValidatorSubsystem.cpp b/Engine/Plugins/Editor/DataValidation/Source/DataValidation/Private/EditorValidatorSubsystem.cpp
new file mode 100644
--- /dev/null
+++ b/Engine/Plugins/Editor/DataValidation/Source/DataValidation/Private/EditorValidatorSubsystem.cpp
@@ -0,0 +1,360 @@
+// Copyright 1998-2019 Epic Games, Inc. All Rights Reserved.
+
+#include "EditorValidatorSubsystem.h"
+
+#include "Editor.h"
+#include "AssetToolsModule.h"
+#include "ObjectTools.h"
+#include "AssetRegistryModule.h"
+#include "EditorUtilityBlueprint.h"
+#include "Settings/EditorLoadingSavingSettings.h"
+#include "UnrealEdGlobals.h"
+#include "Editor/UnrealEdEngine.h"
+#include "CookOnTheSide/CookOnTheFlyServer.h"
+#include "Logging/MessageLog.h"
+#include "Misc/ScopedSlowTask.h"
+#include "AssetData.h"
+
+#define LOCTEXT_NAMESPACE "EditorValidationSubsystem"
+
+UEditorValidatorSubsystem::UEditorValidatorSubsystem()
+	: UEditorSubsystem()
+{
+	bAllowBlueprintValidators = true;
+	bValidateOnSave = true;
+}
+
+void UEditorValidatorSubsystem::Initialize(FSubsystemCollectionBase& Collection)
+{
+	const FAssetRegistryModule& AssetRegistryModule = FModuleManager::LoadModuleChecked<FAssetRegistryModule>(TEXT("AssetRegistry"));
+	if (!AssetRegistryModule.Get().IsLoadingAssets())
+	{
+		RegisterBlueprintValidators();
+	}
+	else
+	{
+		// We are still discovering assets, listen for the completion delegate before building the graph
+		if (!AssetRegistryModule.Get().OnFilesLoaded().IsBoundToObject(this))
+		{
+			AssetRegistryModule.Get().OnFilesLoaded().AddUObject(this, &UEditorValidatorSubsystem::RegisterBlueprintValidators);
+		}
+	}
+
+	// C++ registration
+	TArray<UClass*> ValidatorClasses;
+	GetDerivedClasses(UEditorValidatorBase::StaticClass(), ValidatorClasses);
+	for (UClass* ValidatorClass : ValidatorClasses)
+	{
+		if (!ValidatorClass->HasAllClassFlags(CLASS_Abstract))
+		{
+			UPackage* const ClassPackage = ValidatorClass->GetOuterUPackage();
+			if (ClassPackage)
+			{
+				const FName ModuleName = FPackageName::GetShortFName(ClassPackage->GetFName());
+				if (FModuleManager::Get().IsModuleLoaded(ModuleName))
+				{
+					UEditorValidatorBase* Validator = NewObject<UEditorValidatorBase>(GetTransientPackage(), ValidatorClass);
+					AddValidator(Validator);
+				}
+			}
+		}
+	}
+}
+
+// Rename to BP validators
+void UEditorValidatorSubsystem::RegisterBlueprintValidators()
+{
+	if (bAllowBlueprintValidators)
+	{
+		// Locate all validators (include unloaded)
+		const FAssetRegistryModule& AssetRegistryModule = FModuleManager::LoadModuleChecked<FAssetRegistryModule>(TEXT("AssetRegistry"));
+		TArray<FAssetData> AllBPsAssetData;
+		AssetRegistryModule.Get().GetAssetsByClass(UEditorUtilityBlueprint::StaticClass()->GetFName(), AllBPsAssetData, true);
+
+		for (FAssetData& BPAssetData : AllBPsAssetData)
+		{
+			UClass* ParentClass = nullptr;
+			FString ParentClassName;
+			if (!BPAssetData.GetTagValue(FBlueprintTags::NativeParentClassPath, ParentClassName))
+			{
+				BPAssetData.GetTagValue(FBlueprintTags::ParentClassPath, ParentClassName);
+			}
+
+			if (!ParentClassName.IsEmpty())
+			{
+				UObject* Outer = nullptr;
+				ResolveName(Outer, ParentClassName, false, false);
+				ParentClass = FindObject<UClass>(ANY_PACKAGE, *ParentClassName);
+				if (!ParentClass->IsChildOf(UEditorValidatorBase::StaticClass()))
+				{
+					continue;
+				}
+			}
+
+			// If this object isn't currently loaded, load it
+			UObject* ValidatorObject = BPAssetData.ToSoftObjectPath().ResolveObject();
+			if (ValidatorObject == nullptr)
+			{
+				ValidatorObject = BPAssetData.ToSoftObjectPath().TryLoad();
+			}
+			if (ValidatorObject)
+			{
+				UEditorUtilityBlueprint* ValidatorBlueprint = Cast<UEditorUtilityBlueprint>(ValidatorObject);
+				UEditorValidatorBase* Validator = NewObject<UEditorValidatorBase>(GetTransientPackage(), ValidatorBlueprint->GeneratedClass);
+				AddValidator(Validator);
+			}
+		}
+	}
+}
+
+void UEditorValidatorSubsystem::Deinitialize()
+{
+	CleanupValidators();
+
+	Super::Deinitialize();
+}
+
+void UEditorValidatorSubsystem::AddValidator(UEditorValidatorBase* InValidator)
+{
+	if (InValidator)
+	{
+		Validators.Add(InValidator->GetClass(), InValidator);
+	}
+}
+
+void UEditorValidatorSubsystem::CleanupValidators()
+{
+	Validators.Empty();
+}
+
+EDataValidationResult UEditorValidatorSubsystem::IsObjectValid(UObject* InObject, TArray<FText>& ValidationErrors) const
+{
+	EDataValidationResult Result = EDataValidationResult::NotValidated;
+	
+	if (ensure(InObject))
+	{
+		// First check the class level validation
+		Result = InObject->IsDataValid(ValidationErrors);
+		// If the asset is still valid or there wasn't a class-level validation, keep validating with custom validators
+		if (Result != EDataValidationResult::Invalid)
+		{
+			for (TPair<UClass*, UEditorValidatorBase*> ValidatorPair : Validators)
+			{
+				if (ValidatorPair.Value && ValidatorPair.Value->IsEnabled() && ValidatorPair.Value->CanValidateAsset(InObject))
+				{
+					ValidatorPair.Value->ResetValidationState();
+					EDataValidationResult NewResult = ValidatorPair.Value->ValidateLoadedAsset(InObject, ValidationErrors);
+
+					// Don't accidentally overwrite an invalid result with a valid or not-validated one
+					if(Result != EDataValidationResult::Invalid)
+					{
+						Result = NewResult;
+					}
+
+					ensureMsgf(ValidatorPair.Value->IsValidationStateSet(), TEXT("Validator %s did not include a pass or fail state."), *ValidatorPair.Value->GetClass()->GetName());
+				}
+			}
+		}
+	}
+
+	return Result;
+}
+
+EDataValidationResult UEditorValidatorSubsystem::IsAssetValid(FAssetData& AssetData, TArray<FText>& ValidationErrors) const
+{
+	if (AssetData.IsValid())
+	{
+		UObject* Obj = AssetData.GetAsset();
+		if (Obj)
+		{
+			return IsObjectValid(Obj, ValidationErrors);
+		}
+		return EDataValidationResult::NotValidated;
+	}
+
+	return EDataValidationResult::Invalid;
+}
+
+int32 UEditorValidatorSubsystem::ValidateAssets(TArray<FAssetData> AssetDataList, bool bSkipExcludedDirectories, bool bShowIfNoFailures) const
+{
+	FScopedSlowTask SlowTask(1.0f, LOCTEXT("ValidatingDataTask", "Validating Data..."));
+	SlowTask.Visibility = bShowIfNoFailures ? ESlowTaskVisibility::ForceVisible : ESlowTaskVisibility::Invisible;
+	if (bShowIfNoFailures)
+	{
+		SlowTask.MakeDialogDelayed(.1f);
+	}
+
+	FMessageLog DataValidationLog("AssetCheck");
+
+	int32 NumAdded = 0;
+
+	int32 NumFilesChecked = 0;
+	int32 NumValidFiles = 0;
+	int32 NumInvalidFiles = 0;
+	int32 NumFilesSkipped = 0;
+	int32 NumFilesUnableToValidate = 0;
+
+	int32 NumFilesToValidate = AssetDataList.Num();
+
+	// Now add to map or update as needed
+	for (FAssetData& Data : AssetDataList)
+	{
+		SlowTask.EnterProgressFrame(1.0f / NumFilesToValidate, FText::Format(LOCTEXT("ValidatingFilename", "Validating {0}"), FText::FromString(Data.GetFullName())));
+
+		// Check exclusion path
+		if (bSkipExcludedDirectories && IsPathExcludedFromValidation(Data.PackageName.ToString()))
+		{
+			++NumFilesSkipped;
+			continue;
+		}
+
+		TArray<FText> ValidationErrors;
+		EDataValidationResult Result = IsAssetValid(Data, ValidationErrors);
+		++NumFilesChecked;
+
+		for (const FText& ErrorMsg : ValidationErrors)
+		{
+			DataValidationLog.Error()->AddToken(FTextToken::Create(ErrorMsg));
+		}
+
+		if (Result == EDataValidationResult::Valid)
+		{
+			++NumValidFiles;
+		}
+		else
+		{
+			if (Result == EDataValidationResult::Invalid)
+			{
+				DataValidationLog.Error()->AddToken(FAssetNameToken::Create(Data.PackageName.ToString()))
+					->AddToken(FTextToken::Create(LOCTEXT("InvalidDataResult", "contains invalid data.")));
+				++NumInvalidFiles;
+			}
+			else if (Result == EDataValidationResult::NotValidated)
+			{
+				if (bShowIfNoFailures)
+				{
+					DataValidationLog.Info()->AddToken(FAssetNameToken::Create(Data.PackageName.ToString()))
+						->AddToken(FTextToken::Create(LOCTEXT("NotValidatedDataResult", "has no data validation.")));
+				}
+				++NumFilesUnableToValidate;
+			}
+		}
+	}
+
+	const bool bFailed = (NumInvalidFiles > 0);
+
+	if (bFailed || bShowIfNoFailures)
+	{
+		FFormatNamedArguments Arguments;
+		Arguments.Add(TEXT("Result"), bFailed ? LOCTEXT("Failed", "FAILED") : LOCTEXT("Succeeded", "SUCCEEDED"));
+		Arguments.Add(TEXT("NumChecked"), NumFilesChecked);
+		Arguments.Add(TEXT("NumValid"), NumValidFiles);
+		Arguments.Add(TEXT("NumInvalid"), NumInvalidFiles);
+		Arguments.Add(TEXT("NumSkipped"), NumFilesSkipped);
+		Arguments.Add(TEXT("NumUnableToValidate"), NumFilesUnableToValidate);
+
+		TSharedRef<FTokenizedMessage> ValidationLog = bFailed ? DataValidationLog.Error() : DataValidationLog.Info();
+		ValidationLog->AddToken(FTextToken::Create(FText::Format(LOCTEXT("SuccessOrFailure", "Data validation {Result}."), Arguments)));
+		ValidationLog->AddToken(FTextToken::Create(FText::Format(LOCTEXT("ResultsSummary", "Files Checked: {NumChecked}, Passed: {NumValid}, Failed: {NumInvalid}, Skipped: {NumSkipped}, Unable to validate: {NumUnableToValidate}"), Arguments)));
+
+		DataValidationLog.Open(EMessageSeverity::Info, true);
+	}
+
+	return NumInvalidFiles;
+}
+
+void UEditorValidatorSubsystem::ValidateOnSave(TArray<FAssetData> AssetDataList) const
+{
+	// Only validate if enabled and not auto saving
+	if (!bValidateOnSave || GEditor->IsAutosaving())
+	{
+		return;
+	}
+
+	bool bIsCooking = GIsCookerLoadingPackage;
+	if ((!bValidateAssetsWhileSavingForCook && bIsCooking))
+	{
+		return;
+	}
+
+	FMessageLog DataValidationLog("DataValidation");
+	if (ValidateAssets(AssetDataList, true, false) > 0)
+	{
+		const FText ErrorMessageNotification = FText::Format(
+			LOCTEXT("ValidationFailureNotification", "Validation failed when saving {0}, check Data Validation log"),
+			AssetDataList.Num() == 1 ? FText::FromName(AssetDataList[0].AssetName) : LOCTEXT("MultipleErrors", "multiple assets"));
+		DataValidationLog.Notify(ErrorMessageNotification, EMessageSeverity::Warning, /*bForce=*/ true);
+	}
+}
+
+void UEditorValidatorSubsystem::ValidateSavedPackage(FName PackageName)
+{
+	// Only validate if enabled and not auto saving
+	if (!bValidateOnSave || GEditor->IsAutosaving())
+	{
+		return;
+	}
+
+	// For performance reasons, don't validate when cooking by default. Assumption is we validated when saving previously. 
+	bool bIsCooking = GIsCookerLoadingPackage;
+	if ((!bValidateAssetsWhileSavingForCook && bIsCooking))
+	{
+		return;
+	}
+
+	SavedPackagesToValidate.AddUnique(PackageName);
+
+	GEditor->GetTimerManager()->SetTimerForNextTick(this, &UEditorValidatorSubsystem::ValidateAllSavedPackages);
+}
+
+bool UEditorValidatorSubsystem::IsPathExcludedFromValidation(const FString& Path) const
+{
+	for (const FDirectoryPath& ExcludedPath : ExcludedDirectories)
+	{
+		if (Path.Contains(ExcludedPath.Path))
+		{
+			return true;
+		}
+	}
+
+	return false;
+}
+
+void UEditorValidatorSubsystem::ValidateAllSavedPackages()
+{
+	FAssetRegistryModule& AssetRegistryModule = FModuleManager::LoadModuleChecked<FAssetRegistryModule>("AssetRegistry");
+	IAssetRegistry& AssetRegistry = AssetRegistryModule.Get();
+
+	// Prior to validation, make sure Asset Registry is updated.
+	// DirectoryWatcher is responsible of scanning modified asset files, but validation can be called before.
+	if (SavedPackagesToValidate.Num())
+	{
+		TArray<FString> FilesToScan;
+		FilesToScan.Reserve(SavedPackagesToValidate.Num());
+		for (FName PackageName : SavedPackagesToValidate)
+		{
+			FString PackageFilename;
+			if (FPackageName::FindPackageFileWithoutExtension(FPackageName::LongPackageNameToFilename(PackageName.ToString()), PackageFilename))
+			{
+				FilesToScan.Add(PackageFilename);
+			}
+		}
+		if (FilesToScan.Num())
+		{
+			AssetRegistry.ScanModifiedAssetFiles(FilesToScan);
+		}
+	}
+
+	TArray<FAssetData> Assets;
+	for (FName PackageName : SavedPackagesToValidate)
+	{
+		// We need to query the in-memory data as the disk cache may not be accurate
+		AssetRegistry.GetAssetsByPackageName(PackageName, Assets);
+	}
+
+	ValidateOnSave(Assets);
+
+	SavedPackagesToValidate.Empty();
+}
+
+#undef LOCTEXT_NAMESPACE
\ No newline at end of file
diff --git a/Engine/Plugins/Editor/DataValidation/Source/DataValidation/Public/DataValidationManager.h b/Engine/Plugins/Editor/DataValidation/Source/DataValidation/Public/DataValidationManager.h
--- a/Engine/Plugins/Editor/DataValidation/Source/DataValidation/Public/DataValidationManager.h
+++ b/Engine/Plugins/Editor/DataValidation/Source/DataValidation/Public/DataValidationManager.h
@@ -18,8 +18,8 @@ class UUnitTest;
  * and some misc tasks like local log hooking
  */
 
-UCLASS(config=Editor)
-class DATAVALIDATION_API UDataValidationManager : public UObject
+UCLASS(config=Editor, Deprecated)
+class DATAVALIDATION_API UDEPRECATED_DataValidationManager : public UObject
 {
 	GENERATED_UCLASS_BODY()
 
@@ -29,26 +29,30 @@ public:
 	 *
 	 * @return	Returns the data validation manager
 	 */
-	static UDataValidationManager* Get();
+	UE_DEPRECATED(4.23, "Use GEditor->GetEditorSubsystem<UEditorValidatorSubsystem>() and use the functions on the subsystem instead.")
+	static UDEPRECATED_DataValidationManager* Get();
 
 	/**
 	 * Initialize the data validation manager
 	 */
+	UE_DEPRECATED(4.23, "Use GEditor->GetEditorSubsystem<UEditorValidatorSubsystem>() and use the functions on the subsystem instead.")
 	virtual void Initialize();
 
 	/**
 	 * Destructor for handling removal of log registration
 	 */
-	virtual ~UDataValidationManager() override;
+	virtual ~UDEPRECATED_DataValidationManager() override;
 
 	/**
 	 * @return Returns Valid if the object contains valid data; returns Invalid if the object contains invalid data; returns NotValidated if no validations was performed on the object
 	 */
+	UE_DEPRECATED(4.23, "Use GEditor->GetEditorSubsystem<UEditorValidatorSubsystem>() and use IsObjectValid on the subsystem instead.")
 	virtual EDataValidationResult IsObjectValid(UObject* InObject, TArray<FText>& ValidationErrors) const;
 
 	/**
 	 * @return Returns Valid if the object pointed to by AssetData contains valid data; returns Invalid if the object contains invalid data or does not exist; returns NotValidated if no validations was performed on the object
 	 */
+	UE_DEPRECATED(4.23, "Use GEditor->GetEditorSubsystem<UEditorValidatorSubsystem>() and use IsAssetValid on the subsystem instead.")
 	virtual EDataValidationResult IsAssetValid(FAssetData& AssetData, TArray<FText>& ValidationErrors) const;
 
 	/**
@@ -57,39 +61,44 @@ public:
 	 * @param bShowIfNoFailures If true, will add notifications for files with no validation and display even if everything passes
 	 * @returns Number of assets with validation failures
 	 */
+	UE_DEPRECATED(4.23, "Use GEditor->GetEditorSubsystem<UEditorValidatorSubsystem>() and use IsAssetValid on the subsystem instead.")
 	virtual int32 ValidateAssets(TArray<FAssetData> AssetDataList, bool bSkipExcludedDirectories = true, bool bShowIfNoFailures = true) const;
 
 	/**
 	 * Called to validate from an interactive save
 	 */
+	UE_DEPRECATED(4.23, "Use GEditor->GetEditorSubsystem<UEditorValidatorSubsystem>() and use IsAssetValid on the subsystem instead.")
 	virtual void ValidateOnSave(TArray<FAssetData> AssetDataList) const;
 
 	/**
 	 * Schedule a validation of a saved package, this will activate next frame by default so it can combine them
 	 */
+	UE_DEPRECATED(4.23, "Use GEditor->GetEditorSubsystem<UEditorValidatorSubsystem>() and use ValidateSavedPackage on the subsystem instead.")
 	virtual void ValidateSavedPackage(FName PackageName);
 
 protected:
 	/** 
 	 * @return Returns true if the current Path should be skipped for validation. Returns false otherwise.
 	 */
+	UE_DEPRECATED(4.23, "Use GEditor->GetEditorSubsystem<UEditorValidatorSubsystem>() and use IsPathExcludedFromValidation on the subsystem instead.")
 	virtual bool IsPathExcludedFromValidation(const FString& Path) const;
 
 	/**
 	 * Handles validating all pending save packages
 	 */
+	UE_DEPRECATED(4.23, "Use GEditor->GetEditorSubsystem<UEditorValidatorSubsystem>() and use ValidateAllSavedPackages on the subsystem instead.")
 	void ValidateAllSavedPackages();
 
 	/**
 	 * Directories to ignore for data validation. Useful for test assets
 	 */
-	UPROPERTY(config)
+	UPROPERTY(config, meta = (Deprecated, DeprecationMessage = "UDataValidationManager's ExcludedDirectories is deprecated, use UEditorValidatorSubsystem's ExcludedDirectories instead."))
 	TArray<FDirectoryPath> ExcludedDirectories;
 
 	/**
 	 * Rather it should validate assets on save inside the editor
 	 */
-	UPROPERTY(config)
+	UPROPERTY(config, meta = (Deprecated, DeprecationMessage = "UDataValidationManager's bValidateOnSave is deprecated, use UEditorValidatorSubsystem's bValidateOnSave instead."))
 	bool bValidateOnSave;
 
 	/** List of saved package names to validate next frame */
diff --git a/Engine/Plugins/Editor/DataValidation/Source/DataValidation/Public/EditorValidatorBase.h b/Engine/Plugins/Editor/DataValidation/Source/DataValidation/Public/EditorValidatorBase.h
new file mode 100644
--- /dev/null
+++ b/Engine/Plugins/Editor/DataValidation/Source/DataValidation/Public/EditorValidatorBase.h
@@ -0,0 +1,80 @@
+// Copyright 1998-2019 Epic Games, Inc. All Rights Reserved.
+
+#pragma once
+
+#include "EditorSubsystem.h"
+
+#include "UObject/ObjectMacros.h"
+#include "UObject/UObjectGlobals.h"
+#include "Templates/SubclassOf.h"
+#include "UObject/Interface.h"
+#include "UObject/ObjectMacros.h"
+#include "Logging/LogVerbosity.h"
+
+#include "EditorValidatorBase.generated.h"
+
+/*
+* The EditorValidatorBase is a class which verifies that an asset meets a specific ruleset.
+* It should be used when checking engine-level classes, as UObject::IsDataValid requires
+* overriding the base class. You can create project-specific version of the validator base,
+* with custom logging and enabled logic.
+*
+* C++ and Blueprint validators will be gathered on editor start, while python validators need
+* to register themselves
+*/
+UCLASS(Abstract, Blueprintable, meta = (ShowWorldContextPin))
+class DATAVALIDATION_API UEditorValidatorBase : public UObject
+{
+	GENERATED_BODY()
+
+public:
+	UEditorValidatorBase();
+
+	/** Override this to determine whether or not you can validate a given asset with this validator */
+	UFUNCTION(BlueprintNativeEvent, Category = "Asset Validation")
+	bool CanValidateAsset(UObject* InAsset) const;
+
+	UFUNCTION(BlueprintNativeEvent, Category = "Asset Validation")
+	EDataValidationResult ValidateLoadedAsset(UObject* InAsset, UPARAM(ref) TArray<FText>& ValidationErrors);
+
+	UFUNCTION(BlueprintCallable, Category = "Asset Validation")
+	void AssetFails(UObject* InAsset, const FText& InMessage, UPARAM(ref) TArray<FText>& ValidationErrors);
+
+	UFUNCTION(BlueprintCallable, Category = "Asset Validation")
+	void AssetPasses(UObject* InAsset);
+
+	virtual bool IsEnabled() const
+	{
+		return bIsEnabled;
+	}
+
+	void ResetValidationState();
+
+	bool IsValidationStateSet() const 
+	{
+		return ValidationResult != EDataValidationResult::NotValidated;
+	}
+
+	UFUNCTION(BlueprintCallable, Category = "Asset Validation")
+	EDataValidationResult GetValidationResult() const 
+	{
+		return ValidationResult;
+	}
+
+protected:
+	void LogElapsedTime(FFormatNamedArguments &Arguments);
+
+protected:
+	UPROPERTY(EditAnywhere, Category = "Asset Validation", meta = (BlueprintProtected = "true"))
+	bool bIsEnabled;
+
+private:
+	EDataValidationResult ValidationResult;
+
+	FDateTime ValidationTime;
+};
+
+
+
+
+
diff --git a/Engine/Plugins/Editor/DataValidation/Source/DataValidation/Public/EditorValidatorSubsystem.h b/Engine/Plugins/Editor/DataValidation/Source/DataValidation/Public/EditorValidatorSubsystem.h
new file mode 100644
--- /dev/null
+++ b/Engine/Plugins/Editor/DataValidation/Source/DataValidation/Public/EditorValidatorSubsystem.h
@@ -0,0 +1,119 @@
+// Copyright 1998-2019 Epic Games, Inc. All Rights Reserved.
+
+#pragma once
+
+#include "EditorSubsystem.h"
+
+#include "UObject/ObjectMacros.h"
+#include "UObject/UObjectGlobals.h"
+#include "Templates/SubclassOf.h"
+#include "UObject/Interface.h"
+#include "UObject/ObjectMacros.h"
+#include "EditorValidatorBase.h"
+
+#include "EditorValidatorSubsystem.generated.h"
+
+struct FAssetData;
+
+/**
+ * UEditorValidatorSubsystem manages all the asset validation in the engine. 
+ * The first validation handled is UObject::IsDataValid and its overridden functions.
+ * Those validations require custom classes and are most suited to project-specific
+ * classes. The next validation set is of all registered UEditorValidationBases.
+ * These validators have a function to determine if they can validate a given asset,
+ * and if they are currently enabled. They are good candidates for validating engine
+ * classes or very specific project logic.
+ */
+UCLASS(Config = Editor)
+class DATAVALIDATION_API UEditorValidatorSubsystem : public UEditorSubsystem
+{
+	GENERATED_BODY()
+
+public:
+	UEditorValidatorSubsystem();
+
+	virtual void Initialize(FSubsystemCollectionBase& Collection);
+
+	virtual void Deinitialize();
+
+	/*
+	* Adds a validator to the list, making sure it is a unique instance
+	*/
+	UFUNCTION(BlueprintCallable, Category = "Validation")
+	void AddValidator(UEditorValidatorBase* InValidator);
+
+	/**
+	 * @return Returns Valid if the object contains valid data; returns Invalid if the object contains invalid data; returns NotValidated if no validations was performed on the object
+	 */
+	UFUNCTION(BlueprintCallable, Category = "Asset Validation")
+	virtual EDataValidationResult IsObjectValid(UObject* InObject, TArray<FText>& ValidationErrors) const;
+
+	/**
+	 * @return Returns Valid if the object pointed to by AssetData contains valid data; returns Invalid if the object contains invalid data or does not exist; returns NotValidated if no validations was performed on the object
+	 */
+	UFUNCTION(BlueprintCallable, Category = "Asset Validation")
+	virtual EDataValidationResult IsAssetValid(FAssetData& AssetData, TArray<FText>& ValidationErrors) const;
+
+	/**
+	 * Called to validate assets from either the UI or a commandlet
+	 * @param bSkipExcludedDirectories If true, will not validate files in excluded directories
+	 * @param bShowIfNoFailures If true, will add notifications for files with no validation and display even if everything passes
+	 * @returns Number of assets with validation failures
+	 */
+	UFUNCTION(BlueprintCallable, Category = "Asset Validation")
+	virtual int32 ValidateAssets(TArray<FAssetData> AssetDataList, bool bSkipExcludedDirectories = true, bool bShowIfNoFailures = true) const;
+
+	/**
+	 * Called to validate from an interactive save
+	 */
+	virtual void ValidateOnSave(TArray<FAssetData> AssetDataList) const;
+
+	/**
+	 * Schedule a validation of a saved package, this will activate next frame by default so it can combine them
+	 */
+	virtual void ValidateSavedPackage(FName PackageName);
+
+protected:
+	void CleanupValidators();
+
+	/**
+	 * @return Returns true if the current Path should be skipped for validation. Returns false otherwise.
+	 */
+	virtual bool IsPathExcludedFromValidation(const FString& Path) const;
+
+	/**
+	 * Handles validating all pending save packages
+	 */
+	void ValidateAllSavedPackages();
+
+	void RegisterBlueprintValidators();
+
+protected:
+	/**
+	 * Directories to ignore for data validation. Useful for test assets
+	 */
+	UPROPERTY(config)
+	TArray<FDirectoryPath> ExcludedDirectories;
+
+	/**
+	 * Whether it should validate assets on save inside the editor
+	 */
+	UPROPERTY(config)
+	bool bValidateOnSave;
+
+	/** List of saved package names to validate next frame */
+	TArray<FName> SavedPackagesToValidate;
+
+	UPROPERTY(Transient)
+	TMap<UClass*, UEditorValidatorBase*> Validators;
+
+	/** Specifies whether or not to validate assets on save when saving for a cook */
+	UPROPERTY(config)
+	bool bValidateAssetsWhileSavingForCook;
+
+	/** Specifies whether or not to allow Blueprint validators */
+	UPROPERTY(config)
+	bool bAllowBlueprintValidators;
+
+};
+
diff --git a/Engine/Source/Editor/Blutility/Classes/EditorUtilityBlueprint.h b/Engine/Source/Editor/Blutility/Classes/EditorUtilityBlueprint.h
--- a/Engine/Source/Editor/Blutility/Classes/EditorUtilityBlueprint.h
+++ b/Engine/Source/Editor/Blutility/Classes/EditorUtilityBlueprint.h
@@ -12,7 +12,7 @@
 #include "EditorUtilityBlueprint.generated.h"
 
 UCLASS()
-class UEditorUtilityBlueprint : public UBlueprint
+class BLUTILITY_API UEditorUtilityBlueprint : public UBlueprint
 {
 	GENERATED_UCLASS_BODY()
 
diff --git a/Engine/Source/Runtime/CoreUObject/Public/UObject/NoExportTypes.h b/Engine/Source/Runtime/CoreUObject/Public/UObject/NoExportTypes.h
--- a/Engine/Source/Runtime/CoreUObject/Public/UObject/NoExportTypes.h
+++ b/Engine/Source/Runtime/CoreUObject/Public/UObject/NoExportTypes.h
@@ -1426,6 +1426,19 @@ struct FAutomationExecutionEntry
 	FDateTime Timestamp;
 };
 
+/** Enum used by DataValidation plugin to see if an asset has been validated for correctness (mirrored in UObjectGlobals.h)*/
+UENUM(BlueprintType)
+enum class EDataValidationResult : uint8
+{
+	/** Asset has failed validation */
+	Invalid,
+	/** Asset has passed validation */
+	Valid,
+	/** Asset has not yet been validated */
+	NotValidated
+};
+
+
 /**
  * Direct base class for all UE4 objects
  * @note The full C++ class is located here: Engine\Source\Runtime\CoreUObject\Public\UObject\Object.h
diff --git a/Engine/Source/Runtime/CoreUObject/Public/UObject/UObjectGlobals.h b/Engine/Source/Runtime/CoreUObject/Public/UObject/UObjectGlobals.h
--- a/Engine/Source/Runtime/CoreUObject/Public/UObject/UObjectGlobals.h
+++ b/Engine/Source/Runtime/CoreUObject/Public/UObject/UObjectGlobals.h
@@ -2676,7 +2676,7 @@ namespace UE4CodeGen_Private
 	#define IF_WITH_EDITORONLY_DATA(x, y) y
 #endif
 
-/** Enum used by UDataValidationManager to see if an asset has been validated for correctness */
+/** Enum used by DataValidation plugin to see if an asset has been validated for correctness */
 enum class EDataValidationResult : uint8
 {
 	/** Asset has failed validation */
