Adding PRAGMA_ENABLE_DEPRECATION_WARNINGS
#rnx

diff --git a/Engine/Plugins/Editor/DataValidation/Source/DataValidation/Private/DataValidationManager.cpp b/Engine/Plugins/Editor/DataValidation/Source/DataValidation/Private/DataValidationManager.cpp
--- a/Engine/Plugins/Editor/DataValidation/Source/DataValidation/Private/DataValidationManager.cpp
+++ b/Engine/Plugins/Editor/DataValidation/Source/DataValidation/Private/DataValidationManager.cpp
@@ -75,7 +75,9 @@ EDataValidationResult UDEPRECATED_DataValidationManager::IsAssetValid(FAssetData
 		UObject* Obj = AssetData.GetAsset();
 		if (Obj)
 		{
+			PRAGMA_DISABLE_DEPRECATION_WARNINGS
 			return IsObjectValid(Obj, ValidationErrors);
+			PRAGMA_ENABLE_DEPRECATION_WARNINGS
 		}
 	}
 
@@ -109,14 +111,18 @@ int32 UDEPRECATED_DataValidationManager::ValidateAssets(TArray<FAssetData> Asset
 		SlowTask.EnterProgressFrame(1.0f / NumFilesToValidate, FText::Format(LOCTEXT("ValidatingFilename", "Validating {0}"), FText::FromString(Data.GetFullName())));
 
 		// Check exclusion path
+		PRAGMA_DISABLE_DEPRECATION_WARNINGS
 		if (bSkipExcludedDirectories && IsPathExcludedFromValidation(Data.PackageName.ToString()))
 		{
 			++NumFilesSkipped;
 			continue;
 		}
+		PRAGMA_ENABLE_DEPRECATION_WARNINGS
 
 		TArray<FText> ValidationErrors;
+		PRAGMA_DISABLE_DEPRECATION_WARNINGS
 		EDataValidationResult Result = IsAssetValid(Data, ValidationErrors);
+		PRAGMA_ENABLE_DEPRECATION_WARNINGS
 		++NumFilesChecked;
 
 		for (const FText& ErrorMsg : ValidationErrors)
@@ -179,6 +185,7 @@ void UDEPRECATED_DataValidationManager::ValidateOnSave(TArray<FAssetData> AssetD
 	}
 
 	FMessageLog DataValidationLog("DataValidation");
+	PRAGMA_DISABLE_DEPRECATION_WARNINGS
 	if (ValidateAssets(AssetDataList, true, false) > 0)
 	{
 		const FText ErrorMessageNotification = FText::Format(
@@ -186,6 +193,7 @@ void UDEPRECATED_DataValidationManager::ValidateOnSave(TArray<FAssetData> AssetD
 			AssetDataList.Num() == 1 ? FText::FromName(AssetDataList[0].AssetName) : LOCTEXT("MultipleErrors", "multiple assets"));
 		DataValidationLog.Notify(ErrorMessageNotification, EMessageSeverity::Warning, /*bForce=*/ true);
 	}
+	PRAGMA_ENABLE_DEPRECATION_WARNINGS
 }
 
 void UDEPRECATED_DataValidationManager::ValidateSavedPackage(FName PackageName)
@@ -198,7 +206,9 @@ void UDEPRECATED_DataValidationManager::ValidateSavedPackage(FName PackageName)
 
 	SavedPackagesToValidate.AddUnique(PackageName);
 
+	PRAGMA_DISABLE_DEPRECATION_WARNINGS
 	GEditor->GetTimerManager()->SetTimerForNextTick(this, &UDEPRECATED_DataValidationManager::ValidateAllSavedPackages);
+	PRAGMA_ENABLE_DEPRECATION_WARNINGS
 }
 
 bool UDEPRECATED_DataValidationManager::IsPathExcludedFromValidation(const FString& Path) const
@@ -246,8 +256,9 @@ void UDEPRECATED_DataValidationManager::ValidateAllSavedPackages()
 		AssetRegistry.GetAssetsByPackageName(PackageName, Assets);
 	}
 
+	PRAGMA_DISABLE_DEPRECATION_WARNINGS
 	ValidateOnSave(Assets);
-
+	PRAGMA_ENABLE_DEPRECATION_WARNINGS
 	SavedPackagesToValidate.Empty();
 }
 
