Fixes for non-unity compiles
#rnx

diff --git a/Engine/Plugins/Editor/DataValidation/Source/DataValidation/Private/DataValidationManager.cpp b/Engine/Plugins/Editor/DataValidation/Source/DataValidation/Private/DataValidationManager.cpp
--- a/Engine/Plugins/Editor/DataValidation/Source/DataValidation/Private/DataValidationManager.cpp
+++ b/Engine/Plugins/Editor/DataValidation/Source/DataValidation/Private/DataValidationManager.cpp
@@ -39,7 +39,9 @@ UDEPRECATED_DataValidationManager* UDEPRECATED_DataValidationManager::Get()
 		checkf(GDataValidationManager != nullptr, TEXT("Data validation config value DataValidationManagerClassName is not a subclass of UDataValidationManager."))
 
 		GDataValidationManager->AddToRoot();
+		PRAGMA_DISABLE_DEPRECATION_WARNINGS
 		GDataValidationManager->Initialize();
+		PRAGMA_ENABLE_DEPRECATION_WARNINGS
 	}
 
 	return GDataValidationManager;
diff --git a/Engine/Plugins/Editor/DataValidation/Source/DataValidation/Public/EditorValidatorBase.h b/Engine/Plugins/Editor/DataValidation/Source/DataValidation/Public/EditorValidatorBase.h
--- a/Engine/Plugins/Editor/DataValidation/Source/DataValidation/Public/EditorValidatorBase.h
+++ b/Engine/Plugins/Editor/DataValidation/Source/DataValidation/Public/EditorValidatorBase.h
@@ -10,7 +10,8 @@
 #include "UObject/Interface.h"
 #include "UObject/ObjectMacros.h"
 #include "Logging/LogVerbosity.h"
-
+#include "Internationalization/Text.h"
+#include "UObject/TextProperty.h"
 #include "EditorValidatorBase.generated.h"
 
 /*
diff --git a/Engine/Plugins/Editor/DataValidation/Source/DataValidation/Public/EditorValidatorSubsystem.h b/Engine/Plugins/Editor/DataValidation/Source/DataValidation/Public/EditorValidatorSubsystem.h
--- a/Engine/Plugins/Editor/DataValidation/Source/DataValidation/Public/EditorValidatorSubsystem.h
+++ b/Engine/Plugins/Editor/DataValidation/Source/DataValidation/Public/EditorValidatorSubsystem.h
@@ -10,6 +10,8 @@
 #include "UObject/Interface.h"
 #include "UObject/ObjectMacros.h"
 #include "EditorValidatorBase.h"
+#include "Engine/EngineTypes.h"
+#include "AssetData.h"
 
 #include "EditorValidatorSubsystem.generated.h"
 
