[NO NON-PII COMMIT MESSAGE]

diff --git a/Engine/Plugins/Editor/DataValidation/Source/DataValidation/Private/DataValidationManager.cpp b/Engine/Plugins/Editor/DataValidation/Source/DataValidation/Private/DataValidationManager.cpp
--- a/Engine/Plugins/Editor/DataValidation/Source/DataValidation/Private/DataValidationManager.cpp
+++ b/Engine/Plugins/Editor/DataValidation/Source/DataValidation/Private/DataValidationManager.cpp
@@ -86,7 +86,10 @@ int32 UDataValidationManager::ValidateAssets(TArray<FAssetData> AssetDataList, b
 {
 	FScopedSlowTask SlowTask(1.0f, LOCTEXT("ValidatingDataTask", "Validating Data..."));
 	SlowTask.Visibility = bShowIfNoFailures ? ESlowTaskVisibility::ForceVisible : ESlowTaskVisibility::Invisible;
-	SlowTask.MakeDialog();
+	if (bShowIfNoFailures)
+	{
+		SlowTask.MakeDialogDelayed(.1f);
+	}
 
 	FMessageLog DataValidationLog("DataValidation");
 
