#UE4 Show the LongPackageAssetName instead of object name in validator failure results. This will clean up reporting on some objects like CDOs that have decorations in their name.

diff --git a/Engine/Plugins/Editor/DataValidation/Source/DataValidation/Private/EditorValidatorBase.cpp b/Engine/Plugins/Editor/DataValidation/Source/DataValidation/Private/EditorValidatorBase.cpp
--- a/Engine/Plugins/Editor/DataValidation/Source/DataValidation/Private/EditorValidatorBase.cpp
+++ b/Engine/Plugins/Editor/DataValidation/Source/DataValidation/Private/EditorValidatorBase.cpp
@@ -35,7 +35,8 @@ void UEditorValidatorBase::AssetFails(UObject* InAsset, const FText& InMessage,
 	FFormatNamedArguments Arguments;
 	if (InAsset)
 	{
-		Arguments.Add(TEXT("AssetName"), FText::FromName(InAsset->GetFName()));	
+		FString PackageName = InAsset->GetOutermost()->GetName();
+		Arguments.Add(TEXT("AssetName"), FText::FromString(FPackageName::GetLongPackageAssetName(PackageName)));
 	}
 	Arguments.Add(TEXT("CustomMessage"), InMessage);
 	Arguments.Add(TEXT("ValidatorName"), FText::FromString(GetClass()->GetName()));
