#UE4 Make link in data validation failure toast show the right message log

diff --git a/Engine/Plugins/Editor/DataValidation/Source/DataValidation/Private/EditorValidatorSubsystem.cpp b/Engine/Plugins/Editor/DataValidation/Source/DataValidation/Private/EditorValidatorSubsystem.cpp
--- a/Engine/Plugins/Editor/DataValidation/Source/DataValidation/Private/EditorValidatorSubsystem.cpp
+++ b/Engine/Plugins/Editor/DataValidation/Source/DataValidation/Private/EditorValidatorSubsystem.cpp
@@ -277,7 +277,7 @@ void UEditorValidatorSubsystem::ValidateOnSave(TArray<FAssetData> AssetDataList)
 		return;
 	}
 
-	FMessageLog DataValidationLog("DataValidation");
+	FMessageLog DataValidationLog("AssetCheck");
 	if (ValidateAssets(AssetDataList, true, false) > 0)
 	{
 		const FText ErrorMessageNotification = FText::Format(
