#UE4 Asset Check issues are now reported on their own page when saving assets

diff --git a/Engine/Plugins/Editor/DataValidation/Source/DataValidation/Private/EditorValidatorSubsystem.cpp b/Engine/Plugins/Editor/DataValidation/Source/DataValidation/Private/EditorValidatorSubsystem.cpp
--- a/Engine/Plugins/Editor/DataValidation/Source/DataValidation/Private/EditorValidatorSubsystem.cpp
+++ b/Engine/Plugins/Editor/DataValidation/Source/DataValidation/Private/EditorValidatorSubsystem.cpp
@@ -278,11 +278,12 @@ void UEditorValidatorSubsystem::ValidateOnSave(TArray<FAssetData> AssetDataList)
 	}
 
 	FMessageLog DataValidationLog("AssetCheck");
+	FText SavedAsset = AssetDataList.Num() == 1 ? FText::FromName(AssetDataList[0].AssetName) : LOCTEXT("MultipleErrors", "multiple assets");
+	DataValidationLog.NewPage(FText::Format(LOCTEXT("DataValidationLogPage", "Asset Save: {0}"), SavedAsset));
 	if (ValidateAssets(AssetDataList, true, false) > 0)
 	{
 		const FText ErrorMessageNotification = FText::Format(
-			LOCTEXT("ValidationFailureNotification", "Validation failed when saving {0}, check Data Validation log"),
-			AssetDataList.Num() == 1 ? FText::FromName(AssetDataList[0].AssetName) : LOCTEXT("MultipleErrors", "multiple assets"));
+			LOCTEXT("ValidationFailureNotification", "Validation failed when saving {0}, check Data Validation log"), SavedAsset);
 		DataValidationLog.Notify(ErrorMessageNotification, EMessageSeverity::Warning, /*bForce=*/ true);
 	}
 }
