#UE4 Log out the names of assets being validated in ValidateAssets during Data Validation.

diff --git a/Engine/Plugins/Editor/DataValidation/Source/DataValidation/Private/EditorValidatorBase.cpp b/Engine/Plugins/Editor/DataValidation/Source/DataValidation/Private/EditorValidatorBase.cpp
--- a/Engine/Plugins/Editor/DataValidation/Source/DataValidation/Private/EditorValidatorBase.cpp
+++ b/Engine/Plugins/Editor/DataValidation/Source/DataValidation/Private/EditorValidatorBase.cpp
@@ -11,7 +11,6 @@
 #include "Misc/UObjectToken.h"
 #include "Logging/TokenizedMessage.h"
 
-DEFINE_LOG_CATEGORY_STATIC(LogContentValidation, Log, Log);
 #define LOCTEXT_NAMESPACE "AssetValidation"
 
 UEditorValidatorBase::UEditorValidatorBase()
diff --git a/Engine/Plugins/Editor/DataValidation/Source/DataValidation/Private/EditorValidatorSubsystem.cpp b/Engine/Plugins/Editor/DataValidation/Source/DataValidation/Private/EditorValidatorSubsystem.cpp
--- a/Engine/Plugins/Editor/DataValidation/Source/DataValidation/Private/EditorValidatorSubsystem.cpp
+++ b/Engine/Plugins/Editor/DataValidation/Source/DataValidation/Private/EditorValidatorSubsystem.cpp
@@ -17,6 +17,8 @@
 
 #define LOCTEXT_NAMESPACE "EditorValidationSubsystem"
 
+DEFINE_LOG_CATEGORY(LogContentValidation);
+
 UEditorValidatorSubsystem::UEditorValidatorSubsystem()
 	: UEditorSubsystem()
 {
@@ -199,7 +201,9 @@ int32 UEditorValidatorSubsystem::ValidateAssets(TArray<FAssetData> AssetDataList
 	// Now add to map or update as needed
 	for (FAssetData& Data : AssetDataList)
 	{
-		SlowTask.EnterProgressFrame(1.0f / NumFilesToValidate, FText::Format(LOCTEXT("ValidatingFilename", "Validating {0}"), FText::FromString(Data.GetFullName())));
+		FText ValidatingMessage = FText::Format(LOCTEXT("ValidatingFilename", "Validating {0}"), FText::FromString(Data.GetFullName()));
+		SlowTask.EnterProgressFrame(1.0f / NumFilesToValidate, ValidatingMessage);
+		UE_LOG(LogContentValidation, Display, TEXT("%s"), *ValidatingMessage.ToString());
 
 		// Check exclusion path
 		if (bSkipExcludedDirectories && IsPathExcludedFromValidation(Data.PackageName.ToString()))
diff --git a/Engine/Plugins/Editor/DataValidation/Source/DataValidation/Public/EditorValidatorSubsystem.h b/Engine/Plugins/Editor/DataValidation/Source/DataValidation/Public/EditorValidatorSubsystem.h
--- a/Engine/Plugins/Editor/DataValidation/Source/DataValidation/Public/EditorValidatorSubsystem.h
+++ b/Engine/Plugins/Editor/DataValidation/Source/DataValidation/Public/EditorValidatorSubsystem.h
@@ -12,11 +12,14 @@
 #include "EditorValidatorBase.h"
 #include "Engine/EngineTypes.h"
 #include "AssetData.h"
+#include "Logging/LogMacros.h"
 
 #include "EditorValidatorSubsystem.generated.h"
 
 struct FAssetData;
 
+DECLARE_LOG_CATEGORY_EXTERN(LogContentValidation, Log, All);
+
 /**
  * UEditorValidatorSubsystem manages all the asset validation in the engine. 
  * The first validation handled is UObject::IsDataValid and its overridden functions.
