#UE4 Moved data verification log to avoid emitting when skipping packages due to exclusion rules. The summary messages are no longer errors themselves. Fixed a bug where you wouldnt see the message log when multiple assets were saved

diff --git a/Engine/Plugins/Editor/DataValidation/Source/DataValidation/Private/EditorValidatorSubsystem.cpp b/Engine/Plugins/Editor/DataValidation/Source/DataValidation/Private/EditorValidatorSubsystem.cpp
--- a/Engine/Plugins/Editor/DataValidation/Source/DataValidation/Private/EditorValidatorSubsystem.cpp
+++ b/Engine/Plugins/Editor/DataValidation/Source/DataValidation/Private/EditorValidatorSubsystem.cpp
@@ -203,7 +203,6 @@ int32 UEditorValidatorSubsystem::ValidateAssets(TArray<FAssetData> AssetDataList
 	{
 		FText ValidatingMessage = FText::Format(LOCTEXT("ValidatingFilename", "Validating {0}"), FText::FromString(Data.GetFullName()));
 		SlowTask.EnterProgressFrame(1.0f / NumFilesToValidate, ValidatingMessage);
-		UE_LOG(LogContentValidation, Display, TEXT("%s"), *ValidatingMessage.ToString());
 
 		// Check exclusion path
 		if (bSkipExcludedDirectories && IsPathExcludedFromValidation(Data.PackageName.ToString()))
@@ -212,6 +211,8 @@ int32 UEditorValidatorSubsystem::ValidateAssets(TArray<FAssetData> AssetDataList
 			continue;
 		}
 
+		UE_LOG(LogContentValidation, Display, TEXT("%s"), *ValidatingMessage.ToString());
+
 		TArray<FText> ValidationErrors;
 		EDataValidationResult Result = IsAssetValid(Data, ValidationErrors);
 		++NumFilesChecked;
@@ -257,9 +258,8 @@ int32 UEditorValidatorSubsystem::ValidateAssets(TArray<FAssetData> AssetDataList
 		Arguments.Add(TEXT("NumSkipped"), NumFilesSkipped);
 		Arguments.Add(TEXT("NumUnableToValidate"), NumFilesUnableToValidate);
 
-		TSharedRef<FTokenizedMessage> ValidationLog = bFailed ? DataValidationLog.Error() : DataValidationLog.Info();
-		ValidationLog->AddToken(FTextToken::Create(FText::Format(LOCTEXT("SuccessOrFailure", "Data validation {Result}."), Arguments)));
-		ValidationLog->AddToken(FTextToken::Create(FText::Format(LOCTEXT("ResultsSummary", "Files Checked: {NumChecked}, Passed: {NumValid}, Failed: {NumInvalid}, Skipped: {NumSkipped}, Unable to validate: {NumUnableToValidate}"), Arguments)));
+		DataValidationLog.Info()->AddToken(FTextToken::Create(FText::Format(LOCTEXT("SuccessOrFailure", "Data validation {Result}."), Arguments)));
+		DataValidationLog.Info()->AddToken(FTextToken::Create(FText::Format(LOCTEXT("ResultsSummary", "Files Checked: {NumChecked}, Passed: {NumValid}, Failed: {NumInvalid}, Skipped: {NumSkipped}, Unable to validate: {NumUnableToValidate}"), Arguments)));
 
 		DataValidationLog.Open(EMessageSeverity::Info, true);
 	}
@@ -307,9 +307,12 @@ void UEditorValidatorSubsystem::ValidateSavedPackage(FName PackageName)
 		return;
 	}
 
-	SavedPackagesToValidate.AddUnique(PackageName);
+	if (SavedPackagesToValidate.Num() == 0)
+	{
+		GEditor->GetTimerManager()->SetTimerForNextTick(this, &UEditorValidatorSubsystem::ValidateAllSavedPackages);
+	}
 
-	GEditor->GetTimerManager()->SetTimerForNextTick(this, &UEditorValidatorSubsystem::ValidateAllSavedPackages);
+	SavedPackagesToValidate.AddUnique(PackageName);
 }
 
 bool UEditorValidatorSubsystem::IsPathExcludedFromValidation(const FString& Path) const
