#UE4 Made message about invalid data not an error. The error itself is what should be the error, this is just a message

diff --git a/Engine/Plugins/Editor/DataValidation/Source/DataValidation/Private/DataValidationManager.cpp b/Engine/Plugins/Editor/DataValidation/Source/DataValidation/Private/DataValidationManager.cpp
--- a/Engine/Plugins/Editor/DataValidation/Source/DataValidation/Private/DataValidationManager.cpp
+++ b/Engine/Plugins/Editor/DataValidation/Source/DataValidation/Private/DataValidationManager.cpp
@@ -140,7 +140,7 @@ int32 UDEPRECATED_DataValidationManager::ValidateAssets(TArray<FAssetData> Asset
 		{
 			if (Result == EDataValidationResult::Invalid)
 			{
-				DataValidationLog.Error()->AddToken(FAssetNameToken::Create(Data.PackageName.ToString()))
+				DataValidationLog.Info()->AddToken(FAssetNameToken::Create(Data.PackageName.ToString()))
 					->AddToken(FTextToken::Create(LOCTEXT("InvalidDataResult", "contains invalid data.")));
 				++NumInvalidFiles;
 			}
