#UE4 "Contains invalid data" message is now just an info and not an error. The errors themselves will be errors, this one is just an informational message

#robmerge 11.01

diff --git a/Engine/Plugins/Editor/DataValidation/Source/DataValidation/Private/EditorValidatorSubsystem.cpp b/Engine/Plugins/Editor/DataValidation/Source/DataValidation/Private/EditorValidatorSubsystem.cpp
--- a/Engine/Plugins/Editor/DataValidation/Source/DataValidation/Private/EditorValidatorSubsystem.cpp
+++ b/Engine/Plugins/Editor/DataValidation/Source/DataValidation/Private/EditorValidatorSubsystem.cpp
@@ -230,7 +230,7 @@ int32 UEditorValidatorSubsystem::ValidateAssets(TArray<FAssetData> AssetDataList
 		{
 			if (Result == EDataValidationResult::Invalid)
 			{
-				DataValidationLog.Error()->AddToken(FAssetNameToken::Create(Data.PackageName.ToString()))
+				DataValidationLog.Info()->AddToken(FAssetNameToken::Create(Data.PackageName.ToString()))
 					->AddToken(FTextToken::Create(LOCTEXT("InvalidDataResult", "contains invalid data.")));
 				++NumInvalidFiles;
 			}
