PR #5648: Confirmation dialog before executing folder validation (Contributed by [REDACTED_PERSON])

diff --git a/Engine/Plugins/Editor/DataValidation/Source/DataValidation/Private/DataValidationModule.cpp b/Engine/Plugins/Editor/DataValidation/Source/DataValidation/Private/DataValidationModule.cpp
--- a/Engine/Plugins/Editor/DataValidation/Source/DataValidation/Private/DataValidationModule.cpp
+++ b/Engine/Plugins/Editor/DataValidation/Source/DataValidation/Private/DataValidationModule.cpp
@@ -276,7 +276,25 @@ void FDataValidationModule::CreateDataValidationContentBrowserPathMenu(FMenuBuil
 		LOCTEXT("ValidateAssetsPathTabTitle", "Validate Assets in Folder"),
 		LOCTEXT("ValidateAssetsPathTooltipText", "Runs data validation on the assets in the selected folder."),
 		FSlateIcon(),
-		FUIAction(FExecuteAction::CreateRaw(this, &FDataValidationModule::ValidateFolders, SelectedPaths))
+		FUIAction(FExecuteAction::CreateLambda([this, SelectedPaths]
+		{
+			FString FormattedSelectedPaths;
+			for (int32 i = 0; i < SelectedPaths.Num(); ++i)
+			{
+				FormattedSelectedPaths.Append(SelectedPaths[i]);
+				if (i < SelectedPaths.Num() - 1)
+				{
+					FormattedSelectedPaths.Append(LINE_TERMINATOR);
+				}
+			}
+			FFormatNamedArguments Args;
+			Args.Add(TEXT("Paths"), FText::FromString(FormattedSelectedPaths));
+			const EAppReturnType::Type Result = FMessageDialog::Open(EAppMsgType::YesNo, FText::Format(LOCTEXT("DataValidationConfirmation", "Are you sure you want to proceed with validating the following folders?\n\n{Paths}"), Args));
+			if (Result == EAppReturnType::Yes)
+			{
+				ValidateFolders(SelectedPaths);
+			}
+		}))
 	);
 }
 
