#UE4 Added warnings to data validation

diff --git a/Engine/Plugins/Editor/DataValidation/Source/DataValidation/Private/EditorValidatorBase.cpp b/Engine/Plugins/Editor/DataValidation/Source/DataValidation/Private/EditorValidatorBase.cpp
--- a/Engine/Plugins/Editor/DataValidation/Source/DataValidation/Private/EditorValidatorBase.cpp
+++ b/Engine/Plugins/Editor/DataValidation/Source/DataValidation/Private/EditorValidatorBase.cpp
@@ -52,6 +52,21 @@ void UEditorValidatorBase::AssetFails(UObject* InAsset, const FText& InMessage,
 	ValidationResult = EDataValidationResult::Invalid;
 }
 
+void UEditorValidatorBase::AssetWarning(UObject* InAsset, const FText& InMessage)
+{
+	FFormatNamedArguments Arguments;
+	if (InAsset)
+	{
+		FString PackageName = InAsset->GetOutermost()->GetName();
+		Arguments.Add(TEXT("AssetName"), FText::FromString(FPackageName::GetLongPackageAssetName(PackageName)));
+	}
+	Arguments.Add(TEXT("CustomMessage"), InMessage);
+	Arguments.Add(TEXT("ValidatorName"), FText::FromString(GetClass()->GetName()));
+
+	FText WarningMessage = FText::Format(LOCTEXT("AssetCheck_Message_Display", "{AssetName} warning: {CustomMessage}. ({ValidatorName})"), Arguments);
+	AllWarnings.Add(WarningMessage);
+}
+
 void UEditorValidatorBase::LogElapsedTime(FFormatNamedArguments &Arguments)
 {
 	FDateTime CurrentTime = FDateTime::Now();
@@ -86,6 +101,12 @@ void UEditorValidatorBase::ResetValidationState()
 {
 	ValidationResult = EDataValidationResult::NotValidated;
 	ValidationTime = FDateTime::Now();
+	AllWarnings.Empty();
+}
+
+const TArray<FText>& UEditorValidatorBase::GetAllWarnings() const
+{
+	return AllWarnings;
 }
 
 #undef LOCTEXT_NAMESPACE
\ No newline at end of file
diff --git a/Engine/Plugins/Editor/DataValidation/Source/DataValidation/Private/EditorValidatorSubsystem.cpp b/Engine/Plugins/Editor/DataValidation/Source/DataValidation/Private/EditorValidatorSubsystem.cpp
--- a/Engine/Plugins/Editor/DataValidation/Source/DataValidation/Private/EditorValidatorSubsystem.cpp
+++ b/Engine/Plugins/Editor/DataValidation/Source/DataValidation/Private/EditorValidatorSubsystem.cpp
@@ -129,7 +129,7 @@ void UEditorValidatorSubsystem::CleanupValidators()
 	Validators.Empty();
 }
 
-EDataValidationResult UEditorValidatorSubsystem::IsObjectValid(UObject* InObject, TArray<FText>& ValidationErrors) const
+EDataValidationResult UEditorValidatorSubsystem::IsObjectValid(UObject* InObject, TArray<FText>& ValidationErrors, TArray<FText>& ValidationWarnings) const
 {
 	EDataValidationResult Result = EDataValidationResult::NotValidated;
 	
@@ -153,6 +153,8 @@ EDataValidationResult UEditorValidatorSubsystem::IsObjectValid(UObject* InObject
 						Result = NewResult;
 					}
 
+					ValidationWarnings.Append(ValidatorPair.Value->GetAllWarnings());
+
 					ensureMsgf(ValidatorPair.Value->IsValidationStateSet(), TEXT("Validator %s did not include a pass or fail state."), *ValidatorPair.Value->GetClass()->GetName());
 				}
 			}
@@ -162,14 +164,14 @@ EDataValidationResult UEditorValidatorSubsystem::IsObjectValid(UObject* InObject
 	return Result;
 }
 
-EDataValidationResult UEditorValidatorSubsystem::IsAssetValid(FAssetData& AssetData, TArray<FText>& ValidationErrors) const
+EDataValidationResult UEditorValidatorSubsystem::IsAssetValid(FAssetData& AssetData, TArray<FText>& ValidationErrors, TArray<FText>& ValidationWarnings) const
 {
 	if (AssetData.IsValid())
 	{
 		UObject* Obj = AssetData.GetAsset();
 		if (Obj)
 		{
-			return IsObjectValid(Obj, ValidationErrors);
+			return IsObjectValid(Obj, ValidationErrors, ValidationWarnings);
 		}
 		return EDataValidationResult::NotValidated;
 	}
@@ -214,7 +216,8 @@ int32 UEditorValidatorSubsystem::ValidateAssets(TArray<FAssetData> AssetDataList
 		UE_LOG(LogContentValidation, Display, TEXT("%s"), *ValidatingMessage.ToString());
 
 		TArray<FText> ValidationErrors;
-		EDataValidationResult Result = IsAssetValid(Data, ValidationErrors);
+		TArray<FText> ValidationWarnings;
+		EDataValidationResult Result = IsAssetValid(Data, ValidationErrors, ValidationWarnings);
 		++NumFilesChecked;
 
 		for (const FText& ErrorMsg : ValidationErrors)
@@ -222,6 +225,11 @@ int32 UEditorValidatorSubsystem::ValidateAssets(TArray<FAssetData> AssetDataList
 			DataValidationLog.Error()->AddToken(FTextToken::Create(ErrorMsg));
 		}
 
+		for (const FText& WarningMsg : ValidationWarnings)
+		{
+			DataValidationLog.Warning()->AddToken(FTextToken::Create(WarningMsg));
+		}
+
 		if (Result == EDataValidationResult::Valid)
 		{
 			++NumValidFiles;
diff --git a/Engine/Plugins/Editor/DataValidation/Source/DataValidation/Public/EditorValidatorBase.h b/Engine/Plugins/Editor/DataValidation/Source/DataValidation/Public/EditorValidatorBase.h
--- a/Engine/Plugins/Editor/DataValidation/Source/DataValidation/Public/EditorValidatorBase.h
+++ b/Engine/Plugins/Editor/DataValidation/Source/DataValidation/Public/EditorValidatorBase.h
@@ -44,6 +44,9 @@ public:
 	UFUNCTION(BlueprintCallable, Category = "Asset Validation")
 	void AssetPasses(UObject* InAsset);
 
+	UFUNCTION(BlueprintCallable, Category = "Asset Validation")
+	void AssetWarning(UObject* InAsset, const FText& InMessage);
+
 	virtual bool IsEnabled() const
 	{
 		return bIsEnabled;
@@ -62,6 +65,8 @@ public:
 		return ValidationResult;
 	}
 
+	const TArray<FText>& GetAllWarnings() const;
+
 protected:
 	void LogElapsedTime(FFormatNamedArguments &Arguments);
 
@@ -71,6 +76,7 @@ protected:
 
 private:
 	EDataValidationResult ValidationResult;
+	TArray<FText> AllWarnings;
 
 	FDateTime ValidationTime;
 };
diff --git a/Engine/Plugins/Editor/DataValidation/Source/DataValidation/Public/EditorValidatorSubsystem.h b/Engine/Plugins/Editor/DataValidation/Source/DataValidation/Public/EditorValidatorSubsystem.h
--- a/Engine/Plugins/Editor/DataValidation/Source/DataValidation/Public/EditorValidatorSubsystem.h
+++ b/Engine/Plugins/Editor/DataValidation/Source/DataValidation/Public/EditorValidatorSubsystem.h
@@ -51,13 +51,13 @@ public:
 	 * @return Returns Valid if the object contains valid data; returns Invalid if the object contains invalid data; returns NotValidated if no validations was performed on the object
 	 */
 	UFUNCTION(BlueprintCallable, Category = "Asset Validation")
-	virtual EDataValidationResult IsObjectValid(UObject* InObject, TArray<FText>& ValidationErrors) const;
+	virtual EDataValidationResult IsObjectValid(UObject* InObject, TArray<FText>& ValidationErrors, TArray<FText>& ValidationWarnings) const;
 
 	/**
 	 * @return Returns Valid if the object pointed to by AssetData contains valid data; returns Invalid if the object contains invalid data or does not exist; returns NotValidated if no validations was performed on the object
 	 */
 	UFUNCTION(BlueprintCallable, Category = "Asset Validation")
-	virtual EDataValidationResult IsAssetValid(FAssetData& AssetData, TArray<FText>& ValidationErrors) const;
+	virtual EDataValidationResult IsAssetValid(FAssetData& AssetData, TArray<FText>& ValidationErrors, TArray<FText>& ValidationWarnings) const;
 
 	/**
 	 * Called to validate assets from either the UI or a commandlet
