#UE4 Show validation failure popup when warnings occur, even if there were no errors

diff --git a/Engine/Plugins/Editor/DataValidation/Source/DataValidation/Private/EditorValidatorSubsystem.cpp b/Engine/Plugins/Editor/DataValidation/Source/DataValidation/Private/EditorValidatorSubsystem.cpp
--- a/Engine/Plugins/Editor/DataValidation/Source/DataValidation/Private/EditorValidatorSubsystem.cpp
+++ b/Engine/Plugins/Editor/DataValidation/Source/DataValidation/Private/EditorValidatorSubsystem.cpp
@@ -197,6 +197,7 @@ int32 UEditorValidatorSubsystem::ValidateAssets(TArray<FAssetData> AssetDataList
 	int32 NumInvalidFiles = 0;
 	int32 NumFilesSkipped = 0;
 	int32 NumFilesUnableToValidate = 0;
+	bool bAtLeastOneWarning = false;
 
 	int32 NumFilesToValidate = AssetDataList.Num();
 
@@ -225,13 +226,23 @@ int32 UEditorValidatorSubsystem::ValidateAssets(TArray<FAssetData> AssetDataList
 			DataValidationLog.Error()->AddToken(FTextToken::Create(ErrorMsg));
 		}
 
-		for (const FText& WarningMsg : ValidationWarnings)
+		if (ValidationWarnings.Num() > 0)
 		{
-			DataValidationLog.Warning()->AddToken(FTextToken::Create(WarningMsg));
+			bAtLeastOneWarning = true;
+
+			for (const FText& WarningMsg : ValidationWarnings)
+			{
+				DataValidationLog.Warning()->AddToken(FTextToken::Create(WarningMsg));
+			}
 		}
 
 		if (Result == EDataValidationResult::Valid)
 		{
+			if (ValidationWarnings.Num() > 0)
+			{
+				DataValidationLog.Info()->AddToken(FAssetNameToken::Create(Data.PackageName.ToString()))
+					->AddToken(FTextToken::Create(LOCTEXT("InvalidDataResult", "contains valid data, but has warnings.")));
+			}
 			++NumValidFiles;
 		}
 		else
@@ -256,7 +267,7 @@ int32 UEditorValidatorSubsystem::ValidateAssets(TArray<FAssetData> AssetDataList
 
 	const bool bFailed = (NumInvalidFiles > 0);
 
-	if (bFailed || bShowIfNoFailures)
+	if (bFailed || bAtLeastOneWarning || bShowIfNoFailures)
 	{
 		FFormatNamedArguments Arguments;
 		Arguments.Add(TEXT("Result"), bFailed ? LOCTEXT("Failed", "FAILED") : LOCTEXT("Succeeded", "SUCCEEDED"));
