Added editor preference for whether or not to validate on save.
#rnx

diff --git a/Engine/Plugins/Editor/DataValidation/Source/DataValidation/Private/DataValidationModule.cpp b/Engine/Plugins/Editor/DataValidation/Source/DataValidation/Private/DataValidationModule.cpp
--- a/Engine/Plugins/Editor/DataValidation/Source/DataValidation/Private/DataValidationModule.cpp
+++ b/Engine/Plugins/Editor/DataValidation/Source/DataValidation/Private/DataValidationModule.cpp
@@ -22,6 +22,7 @@
 #include "WorkspaceMenuStructure.h"
 #include "WorkspaceMenuStructureModule.h"
 #include "EditorValidatorSubsystem.h"
+#include "ISettingsModule.h"
 
 #define LOCTEXT_NAMESPACE "DataValidationModule"
 
@@ -78,6 +79,13 @@ void FDataValidationModule::StartupModule()
 
 		// Add save callback
 		UPackage::PackageSavedEvent.AddRaw(this, &FDataValidationModule::OnPackageSaved);
+
+		ISettingsModule& SettingsModule = FModuleManager::LoadModuleChecked<ISettingsModule>("Settings");
+		SettingsModule.RegisterSettings("Editor", "Advanced", "DataValidation",
+			LOCTEXT("DataValidationName", "Data Validation"),
+			LOCTEXT("DataValidationDescription", "Settings related to validating assets in the editor."),
+			GetMutableDefault<UDataValidationSettings>()
+		);
 	}
 }
 
diff --git a/Engine/Plugins/Editor/DataValidation/Source/DataValidation/Private/EditorValidatorSubsystem.cpp b/Engine/Plugins/Editor/DataValidation/Source/DataValidation/Private/EditorValidatorSubsystem.cpp
--- a/Engine/Plugins/Editor/DataValidation/Source/DataValidation/Private/EditorValidatorSubsystem.cpp
+++ b/Engine/Plugins/Editor/DataValidation/Source/DataValidation/Private/EditorValidatorSubsystem.cpp
@@ -19,11 +19,17 @@
 
 DEFINE_LOG_CATEGORY(LogContentValidation);
 
+
+UDataValidationSettings::UDataValidationSettings()
+	: bValidateOnSave(true)
+{
+
+}
+
 UEditorValidatorSubsystem::UEditorValidatorSubsystem()
 	: UEditorSubsystem()
 {
 	bAllowBlueprintValidators = true;
-	bValidateOnSave = true;
 }
 
 void UEditorValidatorSubsystem::Initialize(FSubsystemCollectionBase& Collection)
@@ -289,7 +295,7 @@ int32 UEditorValidatorSubsystem::ValidateAssets(TArray<FAssetData> AssetDataList
 void UEditorValidatorSubsystem::ValidateOnSave(TArray<FAssetData> AssetDataList) const
 {
 	// Only validate if enabled and not auto saving
-	if (!bValidateOnSave || GEditor->IsAutosaving())
+	if (!GetDefault<UDataValidationSettings>()->bValidateOnSave || GEditor->IsAutosaving())
 	{
 		return;
 	}
@@ -314,7 +320,7 @@ void UEditorValidatorSubsystem::ValidateOnSave(TArray<FAssetData> AssetDataList)
 void UEditorValidatorSubsystem::ValidateSavedPackage(FName PackageName)
 {
 	// Only validate if enabled and not auto saving
-	if (!bValidateOnSave || GEditor->IsAutosaving())
+	if (!GetDefault<UDataValidationSettings>()->bValidateOnSave || GEditor->IsAutosaving())
 	{
 		return;
 	}
@@ -384,4 +390,4 @@ void UEditorValidatorSubsystem::ValidateAllSavedPackages()
 	SavedPackagesToValidate.Empty();
 }
 
-#undef LOCTEXT_NAMESPACE
\ No newline at end of file
+#undef LOCTEXT_NAMESPACE
diff --git a/Engine/Plugins/Editor/DataValidation/Source/DataValidation/Public/EditorValidatorSubsystem.h b/Engine/Plugins/Editor/DataValidation/Source/DataValidation/Public/EditorValidatorSubsystem.h
--- a/Engine/Plugins/Editor/DataValidation/Source/DataValidation/Public/EditorValidatorSubsystem.h
+++ b/Engine/Plugins/Editor/DataValidation/Source/DataValidation/Public/EditorValidatorSubsystem.h
@@ -20,6 +20,23 @@ struct FAssetData;
 
 DECLARE_LOG_CATEGORY_EXTERN(LogContentValidation, Log, All);
 
+/**
+* Implements the settings for VR Mode.
+*/
+UCLASS(config = EditorSettings)
+class DATAVALIDATION_API UDataValidationSettings : public UObject
+{
+	GENERATED_BODY()
+
+public:
+	/** Default constructor that sets up CDO properties */
+	UDataValidationSettings();
+
+	/** Display a progress bar while scaling that shows your current scale */
+	UPROPERTY(EditAnywhere, config, Category = "Validation Scenarios")
+	uint32 bValidateOnSave : 1;
+};
+
 /**
  * UEditorValidatorSubsystem manages all the asset validation in the engine. 
  * The first validation handled is UObject::IsDataValid and its overridden functions.
@@ -103,7 +120,7 @@ protected:
 	/**
 	 * Whether it should validate assets on save inside the editor
 	 */
-	UPROPERTY(config)
+	UPROPERTY(config, meta = (DeprecatedProperty, DeprecationMessage = "Use bValidateOnSave on UDataValidationSettings instead."))
 	bool bValidateOnSave;
 
 	/** List of saved package names to validate next frame */
