Fixing a comment on bValidateOnSave
#rnx

diff --git a/Engine/Plugins/Editor/DataValidation/Source/DataValidation/Public/EditorValidatorSubsystem.h b/Engine/Plugins/Editor/DataValidation/Source/DataValidation/Public/EditorValidatorSubsystem.h
--- a/Engine/Plugins/Editor/DataValidation/Source/DataValidation/Public/EditorValidatorSubsystem.h
+++ b/Engine/Plugins/Editor/DataValidation/Source/DataValidation/Public/EditorValidatorSubsystem.h
@@ -32,7 +32,7 @@ public:
 	/** Default constructor that sets up CDO properties */
 	UDataValidationSettings();
 
-	/** Display a progress bar while scaling that shows your current scale */
+	/** Whether or not to validate assets on save */
 	UPROPERTY(EditAnywhere, config, Category = "Validation Scenarios")
 	uint32 bValidateOnSave : 1;
 };
