Fixing class comment also
#rnx

diff --git a/Engine/Plugins/Editor/DataValidation/Source/DataValidation/Public/EditorValidatorSubsystem.h b/Engine/Plugins/Editor/DataValidation/Source/DataValidation/Public/EditorValidatorSubsystem.h
--- a/Engine/Plugins/Editor/DataValidation/Source/DataValidation/Public/EditorValidatorSubsystem.h
+++ b/Engine/Plugins/Editor/DataValidation/Source/DataValidation/Public/EditorValidatorSubsystem.h
@@ -21,7 +21,7 @@ struct FAssetData;
 DECLARE_LOG_CATEGORY_EXTERN(LogContentValidation, Log, All);
 
 /**
-* Implements the settings for VR Mode.
+* Implements the settings for Data Validation 
 */
 UCLASS(config = EditorSettings)
 class DATAVALIDATION_API UDataValidationSettings : public UObject
