Move DataValidationSettings bValidateOnSave from EditorSettings.ini to Editor.ini

diff --git a/Engine/Plugins/Editor/DataValidation/Source/DataValidation/Public/EditorValidatorSubsystem.h b/Engine/Plugins/Editor/DataValidation/Source/DataValidation/Public/EditorValidatorSubsystem.h
--- a/Engine/Plugins/Editor/DataValidation/Source/DataValidation/Public/EditorValidatorSubsystem.h
+++ b/Engine/Plugins/Editor/DataValidation/Source/DataValidation/Public/EditorValidatorSubsystem.h
@@ -23,7 +23,7 @@ DECLARE_LOG_CATEGORY_EXTERN(LogContentValidation, Log, All);
 /**
 * Implements the settings for Data Validation 
 */
-UCLASS(config = EditorSettings)
+UCLASS(config = Editor)
 class DATAVALIDATION_API UDataValidationSettings : public UObject
 {
 	GENERATED_BODY()
