Converted menu: Data Validation
#rnx

diff --git a/Engine/Plugins/Editor/DataValidation/Source/DataValidation/DataValidation.Build.cs b/Engine/Plugins/Editor/DataValidation/Source/DataValidation/DataValidation.Build.cs
--- a/Engine/Plugins/Editor/DataValidation/Source/DataValidation/DataValidation.Build.cs
+++ b/Engine/Plugins/Editor/DataValidation/Source/DataValidation/DataValidation.Build.cs
@@ -34,8 +34,9 @@ public class DataValidation : ModuleRules
 				"BlueprintGraph",
 				"KismetCompiler",
 				"SandboxFile",
-                "Blutility"
-            }
+				"Blutility",
+				"EditorMenus",
+			}
 		);
 	}
 }
diff --git a/Engine/Plugins/Editor/DataValidation/Source/DataValidation/Private/DataValidationModule.cpp b/Engine/Plugins/Editor/DataValidation/Source/DataValidation/Private/DataValidationModule.cpp
--- a/Engine/Plugins/Editor/DataValidation/Source/DataValidation/Private/DataValidationModule.cpp
+++ b/Engine/Plugins/Editor/DataValidation/Source/DataValidation/Private/DataValidationModule.cpp
@@ -6,8 +6,7 @@
 #include "GameFramework/HUD.h"
 
 #include "Framework/Application/SlateApplication.h"
-#include "Framework/MultiBox/MultiBoxBuilder.h"
-#include "Framework/MultiBox/MultiBoxExtender.h"
+#include "EditorMenuSubsystem.h"
 #include "EditorStyleSet.h"
 #include "AssetRegistryModule.h"
 #include "Misc/MessageDialog.h"
@@ -48,12 +47,10 @@ private:
 	// Adds Asset and any assets it depends on to the set DependentAssets
 	void FindAssetDependencies(const FAssetRegistryModule& AssetRegistryModule, const FAssetData& Asset, TSet<FAssetData>& DependentAssets);
 
-	static void DataValidationMenuCreationDelegate(FMenuBuilder& MenuBuilder);
+	void RegisterMenus();
 	static FText Menu_ValidateDataGetTitle();
 	static void Menu_ValidateData();
 
-	TSharedPtr<FExtender> MenuExtender;
-
 	FDelegateHandle ContentBrowserAssetExtenderDelegateHandle;
 	FDelegateHandle ContentBrowserPathExtenderDelegateHandle;
 };
@@ -77,13 +74,7 @@ void FDataValidationModule::StartupModule()
 		ContentBrowserPathExtenderDelegateHandle = CBFolderMenuExtenderDelegates.Last().GetHandle();
 
 		// add the File->DataValidation menu subsection
-
-		// make an extension to add the DataValidation function menu
-		MenuExtender = MakeShareable(new FExtender);
-		MenuExtender->AddMenuExtension("FileLoadAndSave", EExtensionHook::After, nullptr, FMenuExtensionDelegate::CreateStatic(&DataValidationMenuCreationDelegate));
-
-		FLevelEditorModule& LevelEditorModule = FModuleManager::LoadModuleChecked<FLevelEditorModule>("LevelEditor");
-		LevelEditorModule.GetMenuExtensibilityManager()->AddExtender(MenuExtender);
+		FCoreDelegates::OnPostEngineInit.AddRaw(this, &FDataValidationModule::RegisterMenus);
 
 		// Add save callback
 		UPackage::PackageSavedEvent.AddRaw(this, &FDataValidationModule::OnPackageSaved);
@@ -103,26 +94,25 @@ void FDataValidationModule::ShutdownModule()
 		}
 
 		// remove menu extension
-		FLevelEditorModule* LevelEditorModule = FModuleManager::GetModulePtr<FLevelEditorModule>("LevelEditor");
-		if (LevelEditorModule)
-		{
-			LevelEditorModule->GetMenuExtensibilityManager()->RemoveExtender(MenuExtender);
-		}
-		MenuExtender = nullptr;
+		UEditorMenuSubsystem::UnregisterOwner(this);
+		FCoreDelegates::OnPostEngineInit.RemoveAll(this);
 
 		UPackage::PackageSavedEvent.RemoveAll(this);
 	}
 }
 
-void FDataValidationModule::DataValidationMenuCreationDelegate(FMenuBuilder& MenuBuilder)
+void FDataValidationModule::RegisterMenus()
 {
-	MenuBuilder.BeginSection("DataValidation", LOCTEXT("DataValidation", "DataValidation"));
-	MenuBuilder.AddMenuEntry(
+	FEditorMenuOwnerScoped OwnerScoped(this);
+	UEditorMenu* Menu = UEditorMenuSubsystem::Get()->ExtendMenu("LevelEditor.MainMenu.File");
+	FEditorMenuSection& Section = Menu->AddSection("DataValidation", LOCTEXT("DataValidation", "DataValidation"), FEditorMenuInsert("FileLoadAndSave", EEditorMenuInsertType::After));
+	Section.AddEntry(FEditorMenuEntry::InitMenuEntry(
+		"ValidateData",
 		TAttribute<FText>::Create(&Menu_ValidateDataGetTitle),
 		LOCTEXT("ValidateDataTooltip", "Validates all user data in content directory."),
 		FSlateIcon(FEditorStyle::GetStyleSetName(), "DeveloperTools.MenuIcon"),
-		FUIAction(FExecuteAction::CreateStatic(&FDataValidationModule::Menu_ValidateData)));
-	MenuBuilder.EndSection();
+		FUIAction(FExecuteAction::CreateStatic(&FDataValidationModule::Menu_ValidateData))
+	));
 }
 
 FText FDataValidationModule::Menu_ValidateDataGetTitle()
