#UE4 Warnings will now also be counted when returning the number of issues in a validate call. This will cause the notification for validation failures on save to pop up for warnings as well.

diff --git a/Engine/Plugins/Editor/DataValidation/Source/DataValidation/Private/EditorValidatorSubsystem.cpp b/Engine/Plugins/Editor/DataValidation/Source/DataValidation/Private/EditorValidatorSubsystem.cpp
--- a/Engine/Plugins/Editor/DataValidation/Source/DataValidation/Private/EditorValidatorSubsystem.cpp
+++ b/Engine/Plugins/Editor/DataValidation/Source/DataValidation/Private/EditorValidatorSubsystem.cpp
@@ -203,7 +203,7 @@ int32 UEditorValidatorSubsystem::ValidateAssets(TArray<FAssetData> AssetDataList
 	int32 NumInvalidFiles = 0;
 	int32 NumFilesSkipped = 0;
 	int32 NumFilesUnableToValidate = 0;
-	bool bAtLeastOneWarning = false;
+	int32 NumFilesWithWarnings = 0;
 
 	int32 NumFilesToValidate = AssetDataList.Num();
 
@@ -234,7 +234,7 @@ int32 UEditorValidatorSubsystem::ValidateAssets(TArray<FAssetData> AssetDataList
 
 		if (ValidationWarnings.Num() > 0)
 		{
-			bAtLeastOneWarning = true;
+			++NumFilesWithWarnings;
 
 			for (const FText& WarningMsg : ValidationWarnings)
 			{
@@ -272,6 +272,7 @@ int32 UEditorValidatorSubsystem::ValidateAssets(TArray<FAssetData> AssetDataList
 	}
 
 	const bool bFailed = (NumInvalidFiles > 0);
+	const bool bAtLeastOneWarning = (NumFilesWithWarnings > 0);
 
 	if (bFailed || bAtLeastOneWarning || bShowIfNoFailures)
 	{
@@ -289,7 +290,7 @@ int32 UEditorValidatorSubsystem::ValidateAssets(TArray<FAssetData> AssetDataList
 		DataValidationLog.Open(EMessageSeverity::Info, true);
 	}
 
-	return NumInvalidFiles;
+	return NumInvalidFiles + NumFilesWithWarnings;
 }
 
 void UEditorValidatorSubsystem::ValidateOnSave(TArray<FAssetData> AssetDataList) const
diff --git a/Engine/Plugins/Editor/DataValidation/Source/DataValidation/Public/EditorValidatorSubsystem.h b/Engine/Plugins/Editor/DataValidation/Source/DataValidation/Public/EditorValidatorSubsystem.h
--- a/Engine/Plugins/Editor/DataValidation/Source/DataValidation/Public/EditorValidatorSubsystem.h
+++ b/Engine/Plugins/Editor/DataValidation/Source/DataValidation/Public/EditorValidatorSubsystem.h
@@ -80,7 +80,7 @@ public:
 	 * Called to validate assets from either the UI or a commandlet
 	 * @param bSkipExcludedDirectories If true, will not validate files in excluded directories
 	 * @param bShowIfNoFailures If true, will add notifications for files with no validation and display even if everything passes
-	 * @returns Number of assets with validation failures
+	 * @returns Number of assets with validation failures or warnings
 	 */
 	UFUNCTION(BlueprintCallable, Category = "Asset Validation")
 	virtual int32 ValidateAssets(TArray<FAssetData> AssetDataList, bool bSkipExcludedDirectories = true, bool bShowIfNoFailures = true) const;
