#UE4 Add the asset name to the error or warning when validating assets so line-by-line reports will include the name

diff --git a/Engine/Plugins/Editor/DataValidation/Source/DataValidation/Private/EditorValidatorSubsystem.cpp b/Engine/Plugins/Editor/DataValidation/Source/DataValidation/Private/EditorValidatorSubsystem.cpp
--- a/Engine/Plugins/Editor/DataValidation/Source/DataValidation/Private/EditorValidatorSubsystem.cpp
+++ b/Engine/Plugins/Editor/DataValidation/Source/DataValidation/Private/EditorValidatorSubsystem.cpp
@@ -229,7 +229,8 @@ int32 UEditorValidatorSubsystem::ValidateAssets(TArray<FAssetData> AssetDataList
 
 		for (const FText& ErrorMsg : ValidationErrors)
 		{
-			DataValidationLog.Error()->AddToken(FTextToken::Create(ErrorMsg));
+			DataValidationLog.Error()->AddToken(FAssetNameToken::Create(Data.PackageName.ToString()))
+				->AddToken(FTextToken::Create(ErrorMsg));
 		}
 
 		if (ValidationWarnings.Num() > 0)
@@ -238,7 +239,8 @@ int32 UEditorValidatorSubsystem::ValidateAssets(TArray<FAssetData> AssetDataList
 
 			for (const FText& WarningMsg : ValidationWarnings)
 			{
-				DataValidationLog.Warning()->AddToken(FTextToken::Create(WarningMsg));
+				DataValidationLog.Warning()->AddToken(FAssetNameToken::Create(Data.PackageName.ToString()))
+					->AddToken(FTextToken::Create(WarningMsg));
 			}
 		}
 
