Tick the directory watcher instead of manually scanning newly saved assets in the editor validator subsystem to avoid having the "on assets updated" events triggered twice for saved assets.

diff --git a/Engine/Plugins/Editor/DataValidation/Source/DataValidation/Private/EditorValidatorSubsystem.cpp b/Engine/Plugins/Editor/DataValidation/Source/DataValidation/Private/EditorValidatorSubsystem.cpp
--- a/Engine/Plugins/Editor/DataValidation/Source/DataValidation/Private/EditorValidatorSubsystem.cpp
+++ b/Engine/Plugins/Editor/DataValidation/Source/DataValidation/Private/EditorValidatorSubsystem.cpp
@@ -6,6 +6,8 @@
 #include "AssetToolsModule.h"
 #include "ObjectTools.h"
 #include "AssetRegistryModule.h"
+#include "IDirectoryWatcher.h"
+#include "DirectoryWatcherModule.h"
 #include "EditorUtilityBlueprint.h"
 #include "Settings/EditorLoadingSavingSettings.h"
 #include "UnrealEdGlobals.h"
@@ -383,26 +385,14 @@ void UEditorValidatorSubsystem::ValidateAllSavedPackages()
 	FAssetRegistryModule& AssetRegistryModule = FModuleManager::LoadModuleChecked<FAssetRegistryModule>("AssetRegistry");
 	IAssetRegistry& AssetRegistry = AssetRegistryModule.Get();
 
-	// Prior to validation, make sure Asset Registry is updated.
-	// DirectoryWatcher is responsible of scanning modified asset files, but validation can be called before.
-	if (SavedPackagesToValidate.Num())
+	// Prior to validation, make sure Asset Registry is updated. This is done by ticking the DirectoryWatcher module, which 
+	// is responsible of scanning modified asset files.
+	if( !FApp::IsProjectNameEmpty() )
 	{
-		TArray<FString> FilesToScan;
-		FilesToScan.Reserve(SavedPackagesToValidate.Num());
-		for (FName PackageName : SavedPackagesToValidate)
-		{
-			FString PackageFilename;
-			if (FPackageName::FindPackageFileWithoutExtension(FPackageName::LongPackageNameToFilename(PackageName.ToString()), PackageFilename))
-			{
-				FilesToScan.Add(PackageFilename);
-			}
-		}
-		if (FilesToScan.Num())
-		{
-			AssetRegistry.ScanModifiedAssetFiles(FilesToScan);
-		}
+		static FName DirectoryWatcherName("DirectoryWatcher");
+		FDirectoryWatcherModule& DirectoryWatcherModule = FModuleManager::Get().LoadModuleChecked<FDirectoryWatcherModule>(DirectoryWatcherName);
+		DirectoryWatcherModule.Get()->Tick(1.f);
 	}
-
 	// We need to query the in-memory data as the disk cache may not be accurate
 	FARFilter Filter;
 	Filter.PackageNames = SavedPackagesToValidate;
