Change PreSave functions to use IsProceduralChange or IsCooking instead of GIsCookerLoadingPackage so that the extra steps are not taken in EditorDomain saves.
#rnx

diff --git a/Engine/Plugins/Developer/Concert/ConcertSync/ConcertSyncClient/Source/ConcertSyncClient/Private/ConcertClientPackageBridge.cpp b/Engine/Plugins/Developer/Concert/ConcertSync/ConcertSyncClient/Source/ConcertSyncClient/Private/ConcertClientPackageBridge.cpp
--- a/Engine/Plugins/Developer/Concert/ConcertSync/ConcertSyncClient/Source/ConcertSyncClient/Private/ConcertClientPackageBridge.cpp
+++ b/Engine/Plugins/Developer/Concert/ConcertSync/ConcertSyncClient/Source/ConcertSyncClient/Private/ConcertClientPackageBridge.cpp
@@ -126,10 +126,9 @@ bool& FConcertClientPackageBridge::GetIgnoreLocalDiscardRef()
 
 void FConcertClientPackageBridge::HandlePackagePreSave(UPackage* Package, FObjectPreSaveContext ObjectSaveContext)
 {
-	// Ignore package operations fired by the cooker (cook on the fly).
-	if (GIsCookerLoadingPackage)
+	// Only execute if this is a user save
+	if (ObjectSaveContext.IsProceduralSave())
 	{
-		check(IsInGameThread()); // We expect the cooker to call us on the game thread otherwise, we can have concurrency issues.
 		return;
 	}
 
@@ -166,10 +165,9 @@ void FConcertClientPackageBridge::HandlePackagePreSave(UPackage* Package, FObjec
 
 void FConcertClientPackageBridge::HandlePackageSaved(const FString& PackageFilename, UPackage* Package, FObjectPostSaveContext ObjectSaveContext)
 {
-	// Ignore package operations fired by the cooker (cook on the fly).
-	if (GIsCookerLoadingPackage)
+	// Only execute if this is a user save
+	if (ObjectSaveContext.IsProceduralSave())
 	{
-		check(IsInGameThread()); // We expect the cooker to call us on the game thread otherwise, we can have concurrency issues.
 		return;
 	}
 
diff --git a/Engine/Plugins/Editor/AssetSearch/Source/Private/AssetSearchManager.cpp b/Engine/Plugins/Editor/AssetSearch/Source/Private/AssetSearchManager.cpp
--- a/Engine/Plugins/Editor/AssetSearch/Source/Private/AssetSearchManager.cpp
+++ b/Engine/Plugins/Editor/AssetSearch/Source/Private/AssetSearchManager.cpp
@@ -452,8 +452,8 @@ void FAssetSearchManager::HandlePackageSaved(const FString& PackageFilename, UPa
 {
 	check(IsInGameThread());
 
-	// Ignore package operations fired by the cooker (cook on the fly).
-	if (GIsCookerLoadingPackage)
+	// Only execute if this is a user save
+	if (ObjectSaveContext.IsProceduralSave())
 	{
 		return;
 	}
diff --git a/Engine/Plugins/Editor/DataValidation/Source/DataValidation/Private/DataValidationModule.cpp b/Engine/Plugins/Editor/DataValidation/Source/DataValidation/Private/DataValidationModule.cpp
--- a/Engine/Plugins/Editor/DataValidation/Source/DataValidation/Private/DataValidationModule.cpp
+++ b/Engine/Plugins/Editor/DataValidation/Source/DataValidation/Private/DataValidationModule.cpp
@@ -313,7 +313,7 @@ void FDataValidationModule::OnPackageSaved(const FString& PackageFileName, UPack
 	UEditorValidatorSubsystem* EditorValidationSubsystem = GEditor->GetEditorSubsystem<UEditorValidatorSubsystem>();
 	if (EditorValidationSubsystem && Package)
 	{
-		EditorValidationSubsystem->ValidateSavedPackage(Package->GetFName());
+		EditorValidationSubsystem->ValidateSavedPackage(Package->GetFName(), ObjectSaveContext.IsProceduralSave());
 	}
 }
 
diff --git a/Engine/Plugins/Editor/DataValidation/Source/DataValidation/Private/EditorValidatorSubsystem.cpp b/Engine/Plugins/Editor/DataValidation/Source/DataValidation/Private/EditorValidatorSubsystem.cpp
--- a/Engine/Plugins/Editor/DataValidation/Source/DataValidation/Private/EditorValidatorSubsystem.cpp
+++ b/Engine/Plugins/Editor/DataValidation/Source/DataValidation/Private/EditorValidatorSubsystem.cpp
@@ -326,6 +326,11 @@ int32 UEditorValidatorSubsystem::ValidateAssets(TArray<FAssetData> AssetDataList
 }
 
 void UEditorValidatorSubsystem::ValidateOnSave(TArray<FAssetData> AssetDataList) const
+{
+	ValidateOnSave(AssetDataList, false /* bProceduralSave */);
+}
+
+void UEditorValidatorSubsystem::ValidateOnSave(TArray<FAssetData> AssetDataList, bool bProceduralSave) const
 {
 	// Only validate if enabled and not auto saving
 	if (!GetDefault<UDataValidationSettings>()->bValidateOnSave || GEditor->IsAutosaving())
@@ -333,8 +338,7 @@ void UEditorValidatorSubsystem::ValidateOnSave(TArray<FAssetData> AssetDataList)
 		return;
 	}
 
-	bool bIsCooking = GIsCookerLoadingPackage;
-	if ((!bValidateAssetsWhileSavingForCook && bIsCooking))
+	if ((!bValidateAssetsWhileSavingForCook && bProceduralSave))
 	{
 		return;
 	}
@@ -351,6 +355,11 @@ void UEditorValidatorSubsystem::ValidateOnSave(TArray<FAssetData> AssetDataList)
 }
 
 void UEditorValidatorSubsystem::ValidateSavedPackage(FName PackageName)
+{
+	ValidateSavedPackage(PackageName, false /* bProceduralSave */);
+}
+
+void UEditorValidatorSubsystem::ValidateSavedPackage(FName PackageName, bool bProceduralSave)
 {
 	// Only validate if enabled and not auto saving
 	if (!GetDefault<UDataValidationSettings>()->bValidateOnSave || GEditor->IsAutosaving())
@@ -358,9 +367,8 @@ void UEditorValidatorSubsystem::ValidateSavedPackage(FName PackageName)
 		return;
 	}
 
-	// For performance reasons, don't validate when cooking by default. Assumption is we validated when saving previously. 
-	bool bIsCooking = GIsCookerLoadingPackage;
-	if ((!bValidateAssetsWhileSavingForCook && bIsCooking))
+	// For performance reasons, don't validate when making a procedural save by default. Assumption is we validated when saving previously. 
+	if ((!bValidateAssetsWhileSavingForCook && bProceduralSave))
 	{
 		return;
 	}
@@ -409,7 +417,8 @@ void UEditorValidatorSubsystem::ValidateAllSavedPackages()
 	TArray<FAssetData> Assets;
 	AssetRegistry.GetAssets(Filter, Assets);
 
-	ValidateOnSave(MoveTemp(Assets));
+	bool bProceduralSave = false; // The optional suppression for ProceduralSaves was checked before adding to SavedPackagesToValidate
+	ValidateOnSave(MoveTemp(Assets), bProceduralSave);
 
 	SavedPackagesToValidate.Empty();
 }
diff --git a/Engine/Plugins/Editor/DataValidation/Source/DataValidation/Public/EditorValidatorSubsystem.h b/Engine/Plugins/Editor/DataValidation/Source/DataValidation/Public/EditorValidatorSubsystem.h
--- a/Engine/Plugins/Editor/DataValidation/Source/DataValidation/Public/EditorValidatorSubsystem.h
+++ b/Engine/Plugins/Editor/DataValidation/Source/DataValidation/Public/EditorValidatorSubsystem.h
@@ -87,15 +87,27 @@ public:
 	UFUNCTION(BlueprintCallable, Category = "Asset Validation")
 	virtual int32 ValidateAssets(TArray<FAssetData> AssetDataList, bool bSkipExcludedDirectories = true, bool bShowIfNoFailures = true) const;
 
+	UE_DEPRECATED(5.0, "Use version that takes bProceduralSave instead")
+	virtual void ValidateOnSave(TArray<FAssetData> AssetDataList) const;
+
 	/**
 	 * Called to validate from an interactive save
+	 * 
+	 * @param AssetDataList List of Assets to validate
+	 * @param bProceduralSave True iff this is e.g. a save from the editor, false if it is a procedural change from e.g. cooking
 	 */
-	virtual void ValidateOnSave(TArray<FAssetData> AssetDataList) const;
+	virtual void ValidateOnSave(TArray<FAssetData> AssetDataList, bool bProceduralSave) const;
+
+	UE_DEPRECATED(5.0, "Use version that takes bProceduralSave instead")
+	virtual void ValidateSavedPackage(FName PackageName);
 
 	/**
 	 * Schedule a validation of a saved package, this will activate next frame by default so it can combine them
+	 *
+	 * @param PackageName Package to validate
+	 * @param bProceduralSave True iff this is e.g. a save from the editor, false if it is a procedural change from e.g. cooking
 	 */
-	virtual void ValidateSavedPackage(FName PackageName);
+	virtual void ValidateSavedPackage(FName PackageName, bool bProceduralSave);
 
 protected:
 	void CleanupValidators();
diff --git a/Engine/Plugins/Editor/GameplayTagsEditor/Source/GameplayTagsEditor/Private/GameplayTagsEditorModule.cpp b/Engine/Plugins/Editor/GameplayTagsEditor/Source/GameplayTagsEditor/Private/GameplayTagsEditorModule.cpp
--- a/Engine/Plugins/Editor/GameplayTagsEditor/Source/GameplayTagsEditor/Private/GameplayTagsEditorModule.cpp
+++ b/Engine/Plugins/Editor/GameplayTagsEditor/Source/GameplayTagsEditor/Private/GameplayTagsEditorModule.cpp
@@ -149,7 +149,7 @@ public:
 
 	void OnPackageSaved(const FString& PackageFileName, UPackage* Package, FObjectPostSaveContext ObjectSaveContext)
 	{
-		if (GIsEditor && !IsRunningCommandlet())
+		if (GIsEditor && !ObjectSaveContext.IsProceduralSave())
 		{
 			UGameplayTagsManager& Manager = UGameplayTagsManager::Get();
 
diff --git a/Engine/Source/Editor/LandscapeEditor/Private/LandscapeEdMode.cpp b/Engine/Source/Editor/LandscapeEditor/Private/LandscapeEdMode.cpp
--- a/Engine/Source/Editor/LandscapeEditor/Private/LandscapeEdMode.cpp
+++ b/Engine/Source/Editor/LandscapeEditor/Private/LandscapeEdMode.cpp
@@ -672,8 +672,8 @@ void FEdModeLandscape::Exit()
 
 void FEdModeLandscape::OnPreSaveWorld(const class UWorld* InWorld, FObjectPreSaveContext ObjectSaveContext)
 {
-	// Avoid doing this during cooking to keep determinism and we don't want to do this on GameWorlds.
-	if (!InWorld->IsGameWorld() && !GIsCookerLoadingPackage)
+	// Avoid doing this during procedural saves to keep determinism and we don't want to do this on GameWorlds.
+	if (!InWorld->IsGameWorld() && !ObjectSaveContext.IsProceduralSave())
 	{
 		ULandscapeInfoMap& LandscapeInfoMap = ULandscapeInfoMap::GetLandscapeInfoMap(InWorld);
 		for (const TPair<FGuid, ULandscapeInfo*>& Pair : LandscapeInfoMap.Map)
diff --git a/Engine/Source/Runtime/Engine/Private/Animation/AnimSequence.cpp b/Engine/Source/Runtime/Engine/Private/Animation/AnimSequence.cpp
--- a/Engine/Source/Runtime/Engine/Private/Animation/AnimSequence.cpp
+++ b/Engine/Source/Runtime/Engine/Private/Animation/AnimSequence.cpp
@@ -816,7 +816,7 @@ void UAnimSequence::PreSave(FObjectPreSaveContext ObjectSaveContext)
 
 	WaitOnExistingCompression(); // Wait on updated data
 
-	if (!ObjectSaveContext.IsCooking())
+	if (!ObjectSaveContext.IsProceduralSave())
 	{
 		UpdateRetargetSourceAsset();
 	}
diff --git a/Engine/Source/Runtime/Engine/Private/Animation/PoseAsset.cpp b/Engine/Source/Runtime/Engine/Private/Animation/PoseAsset.cpp
--- a/Engine/Source/Runtime/Engine/Private/Animation/PoseAsset.cpp
+++ b/Engine/Source/Runtime/Engine/Private/Animation/PoseAsset.cpp
@@ -901,7 +901,7 @@ void UPoseAsset::PreSave(const ITargetPlatform* TargetPlatform)
 void UPoseAsset::PreSave(FObjectPreSaveContext ObjectSaveContext)
 {
 #if WITH_EDITOR
-	if (!ObjectSaveContext.IsCooking())
+	if (!ObjectSaveContext.IsProceduralSave())
 	{
 		UpdateRetargetSourceAsset();
 	}
diff --git a/Engine/Source/Runtime/Engine/Private/HLOD/HLODEngineSubsystem.cpp b/Engine/Source/Runtime/Engine/Private/HLOD/HLODEngineSubsystem.cpp
--- a/Engine/Source/Runtime/Engine/Private/HLOD/HLODEngineSubsystem.cpp
+++ b/Engine/Source/Runtime/Engine/Private/HLOD/HLODEngineSubsystem.cpp
@@ -169,7 +169,7 @@ bool UHLODEngineSubsystem::CleanupHLOD(ALODActor* InLODActor)
 void UHLODEngineSubsystem::OnPreSaveWorld(UWorld* InWorld, FObjectPreSaveContext ObjectSaveContext)
 {
 	// When cooking, make sure that the LODActors are not transient
-	if (InWorld && InWorld->PersistentLevel && GIsCookerLoadingPackage)
+	if (InWorld && InWorld->PersistentLevel && ObjectSaveContext.IsCooking())
 	{
 		for (TActorIterator<ALODActor> It(InWorld); It; ++It)
 		{
diff --git a/Engine/Source/Runtime/Engine/Private/HLODProxy.cpp b/Engine/Source/Runtime/Engine/Private/HLODProxy.cpp
--- a/Engine/Source/Runtime/Engine/Private/HLODProxy.cpp
+++ b/Engine/Source/Runtime/Engine/Private/HLODProxy.cpp
@@ -171,7 +171,7 @@ void UHLODProxy::PreSave(FObjectPreSaveContext ObjectSaveContext)
 
 	// Always rebuild key on save here.
 	// We don't do this while cooking as keys rely on platform derived data which is context-dependent during cook
-	if (!GIsCookerLoadingPackage)
+	if (!ObjectSaveContext.IsCooking())
 	{
 		if (GetDefault<UHierarchicalLODSettings>()->bSaveLODActorsToHLODPackages)
 		{
diff --git a/Engine/Source/Runtime/Engine/Private/LODActor.cpp b/Engine/Source/Runtime/Engine/Private/LODActor.cpp
--- a/Engine/Source/Runtime/Engine/Private/LODActor.cpp
+++ b/Engine/Source/Runtime/Engine/Private/LODActor.cpp
@@ -1397,7 +1397,7 @@ void ALODActor::PreSave(FObjectPreSaveContext ObjectSaveContext)
 
 	// Always rebuild key on save here.
 	// We don't do this while cooking as keys rely on platform derived data which is context-dependent during cook
-	if(!GIsCookerLoadingPackage)
+	if(!ObjectSaveContext.IsCooking())
 	{
 		const bool bMustUndoLevelTransform = false; // In the save process, the level transform is already removed
 		Key = UHLODProxy::GenerateKeyForActor(this, bMustUndoLevelTransform);
diff --git a/Engine/Source/Runtime/Engine/Private/Texture.cpp b/Engine/Source/Runtime/Engine/Private/Texture.cpp
--- a/Engine/Source/Runtime/Engine/Private/Texture.cpp
+++ b/Engine/Source/Runtime/Engine/Private/Texture.cpp
@@ -598,7 +598,7 @@ void UTexture::PreSave(FObjectPreSaveContext ObjectSaveContext)
 		UpdateResource();
 	}
 
-	if (!GEngine->IsAutosaving() && !ObjectSaveContext.IsCooking())
+	if (!GEngine->IsAutosaving() && !ObjectSaveContext.IsProceduralSave())
 	{
 		GWarn->StatusUpdate(0, 0, FText::Format(NSLOCTEXT("UnrealEd", "SavingPackage_CompressingSourceArt", "Compressing source art for texture:  {0}"), FText::FromString(GetName())));
 		Source.Compress();
diff --git a/Engine/Source/Runtime/Engine/Private/World.cpp b/Engine/Source/Runtime/Engine/Private/World.cpp
--- a/Engine/Source/Runtime/Engine/Private/World.cpp
+++ b/Engine/Source/Runtime/Engine/Private/World.cpp
@@ -1211,7 +1211,7 @@ void UWorld::PreSaveRoot(FObjectPreSaveRootContext ObjectSaveContext)
 
 #if WITH_EDITOR
 	// if we are cooking this world, convert its persistent level to internal actors before doing so
-	if (GIsCookerLoadingPackage)
+	if (ObjectSaveContext.IsCooking())
 	{
 		PersistentLevel->DetachAttachAllActorsPackages(/*bReattach*/false);
 	}
@@ -1248,7 +1248,7 @@ void UWorld::PostSaveRoot( FObjectPostSaveRootContext ObjectSaveContext )
 #if WITH_EDITOR
 	// if we are cooking this world, convert its persistent level back its proper loading strategy if needed
 	// NOTE: can't use bCleanupIsRequired since we don't want to unregister component if they were registered...
-	if (GIsCookerLoadingPackage)
+	if (ObjectSaveContext.IsCooking())
 	{
 		PersistentLevel->DetachAttachAllActorsPackages(/*bReattach*/true);
 	}
diff --git a/Engine/Source/Runtime/Landscape/Private/Landscape.cpp b/Engine/Source/Runtime/Landscape/Private/Landscape.cpp
--- a/Engine/Source/Runtime/Landscape/Private/Landscape.cpp
+++ b/Engine/Source/Runtime/Landscape/Private/Landscape.cpp
@@ -2800,7 +2800,7 @@ void ALandscape::PreSave(FObjectPreSaveContext ObjectSaveContext)
 {
 	Super::PreSave(ObjectSaveContext);
 	//ULandscapeInfo* Info = GetLandscapeInfo();
-	//if (GIsEditor && Info && !IsRunningCommandlet())
+	//if (GIsEditor && Info && !ObjectSaveContext.IsProceduralSave())
 	//{
 	//	for (TSet<ALandscapeProxy*>::TIterator It(Info->Proxies); It; ++It)
 	//	{
diff --git a/Engine/Source/Runtime/Landscape/Private/LandscapeCollision.cpp b/Engine/Source/Runtime/Landscape/Private/LandscapeCollision.cpp
--- a/Engine/Source/Runtime/Landscape/Private/LandscapeCollision.cpp
+++ b/Engine/Source/Runtime/Landscape/Private/LandscapeCollision.cpp
@@ -2490,7 +2490,7 @@ void ULandscapeHeightfieldCollisionComponent::PreSave(FObjectPreSaveContext Obje
 {
 	Super::PreSave(ObjectSaveContext);
 
-	if (!IsRunningCommandlet())
+	if (!ObjectSaveContext.IsProceduralSave())
 	{
 #if WITH_EDITOR
 		ALandscapeProxy* Proxy = GetLandscapeProxy();
diff --git a/Engine/Source/Runtime/UMG/Private/WidgetTree.cpp b/Engine/Source/Runtime/UMG/Private/WidgetTree.cpp
--- a/Engine/Source/Runtime/UMG/Private/WidgetTree.cpp
+++ b/Engine/Source/Runtime/UMG/Private/WidgetTree.cpp
@@ -249,11 +249,11 @@ void UWidgetTree::PreSave(const class ITargetPlatform* TargetPlatform)
 
 void UWidgetTree::PreSave(FObjectPreSaveContext ObjectSaveContext)
 {
-	if (!ObjectSaveContext.IsCooking())
+	if (!ObjectSaveContext.IsCooking()) // AllWidgets is not used in cooked builds
 	{
 #if WITH_EDITORONLY_DATA
+		// Update AllWidgets array used during serialization to match the current state of the WidgetTree
 		AllWidgets.Empty();
-
 		GetAllWidgets(AllWidgets);
 #endif
 	}
