Limit changelist data validation to direct dependencies only

diff --git a/Engine/Plugins/Editor/DataValidation/Source/DataValidation/Private/DataValidationChangelist.cpp b/Engine/Plugins/Editor/DataValidation/Source/DataValidation/Private/DataValidationChangelist.cpp
--- a/Engine/Plugins/Editor/DataValidation/Source/DataValidation/Private/DataValidationChangelist.cpp
+++ b/Engine/Plugins/Editor/DataValidation/Source/DataValidation/Private/DataValidationChangelist.cpp
@@ -10,10 +10,12 @@
 
 #define LOCTEXT_NAMESPACE "DataValidationChangelist"
 
-void GatherDependencies(const FName& InPackageName, TSet<FName>& OutDependencies)
+void GatherDependencies(const FName& InPackageName, TSet<FName>& OutDependencies, bool bGetFullDependencies)
 {
 	if (OutDependencies.Contains(InPackageName))
+	{
 		return;
+	}
 
 	OutDependencies.Add(InPackageName);
 
@@ -23,7 +25,14 @@ void GatherDependencies(const FName& InPackageName, TSet<FName>& OutDependencies
 
 	for (const FName& PackageDependency : Dependencies)
 	{
-		GatherDependencies(PackageDependency, OutDependencies);
+		if (bGetFullDependencies)
+		{
+			GatherDependencies(PackageDependency, OutDependencies, bGetFullDependencies);
+		}
+		else if (!OutDependencies.Contains(PackageDependency))
+		{
+			OutDependencies.Add(PackageDependency);
+		}
 	}
 }
 
@@ -71,6 +80,7 @@ EDataValidationResult UDataValidationChangelist::IsDataValid(TArray<FText>& Vali
 	// Gather dependencies of every file in the changelist
 	TArray<FName> FilesInChangelist;
 	TSet<FName> ExternalDependenciesSet;
+	const bool bGetFullDependencyHierarchy = false;
 
 	for (const FSourceControlStateRef& File : ChangelistState->GetFilesStates())
 	{
@@ -84,7 +94,7 @@ EDataValidationResult UDataValidationChangelist::IsDataValid(TArray<FText>& Vali
 		if (FPackageName::TryConvertFilenameToLongPackageName(File->GetFilename(), PackageName))
 		{
 			FilesInChangelist.Add(*PackageName);
-			GatherDependencies(*PackageName, ExternalDependenciesSet);
+			GatherDependencies(*PackageName, ExternalDependenciesSet, bGetFullDependencyHierarchy);
 		}
 	}
 
