Adds DirtyFilesChangelistValidator to verify there is no unsaved modifications when submitting a changelist.

#changelist validated

diff --git a/Engine/Plugins/Editor/DataValidation/Source/DataValidation/Private/DirtyFilesChangelistValidator.cpp b/Engine/Plugins/Editor/DataValidation/Source/DataValidation/Private/DirtyFilesChangelistValidator.cpp
new file mode 100644
--- /dev/null
+++ b/Engine/Plugins/Editor/DataValidation/Source/DataValidation/Private/DirtyFilesChangelistValidator.cpp
@@ -0,0 +1,74 @@
+// Copyright Epic Games, Inc. All Rights Reserved.
+
+#include "DirtyFilesChangelistValidator.h"
+
+#include "Algo/Transform.h"
+#include "DataValidationChangelist.h"
+#include "FileHelpers.h"
+#include "ISourceControlChangelistState.h"
+#include "ISourceControlModule.h"
+#include "ISourceControlProvider.h"
+#include "Misc/Paths.h"
+
+#define LOCTEXT_NAMESPACE "DirtyFilesChangelistValidation"
+
+FString UDirtyFilesChangelistValidator::GetPackagePath(const UPackage* InPackage)
+{
+	if (InPackage == nullptr)
+	{
+		return TEXT("");
+	}
+
+	const FString LocalFullPath(InPackage->GetLoadedPath().GetLocalFullPath());
+
+	if (LocalFullPath.IsEmpty())
+	{
+		return TEXT("");
+	}
+
+	return FPaths::ConvertRelativePathToFull(LocalFullPath);
+}
+
+bool UDirtyFilesChangelistValidator::CanValidateAsset_Implementation(UObject* InAsset) const
+{
+	return (InAsset != nullptr) && (UDataValidationChangelist::StaticClass() == InAsset->GetClass());
+}
+
+EDataValidationResult UDirtyFilesChangelistValidator::ValidateLoadedAsset_Implementation(UObject* InAsset, TArray<FText>& ValidationErrors)
+{
+	check(UDataValidationChangelist::StaticClass() == InAsset->GetClass());
+
+	UDataValidationChangelist* DataValidationChangelist = Cast<UDataValidationChangelist>(InAsset);
+
+	if ((DataValidationChangelist == nullptr) || (!DataValidationChangelist->Changelist.IsValid()))
+	{
+		return EDataValidationResult::Valid;
+	}
+
+	// Retrieve files contained in the changelist
+	ISourceControlProvider& Provider = ISourceControlModule::Get().GetProvider();
+	FSourceControlChangelistStatePtr ChangelistStatePtr = Provider.GetState(DataValidationChangelist->Changelist.ToSharedRef(), EStateCacheUsage::Use);
+	const TArray<FSourceControlStateRef>& FileStates = ChangelistStatePtr->GetFilesStates();
+	
+	// Retrieve current unsaved packages
+	TArray<UPackage*> DirtyPackages;
+	TSet<FString> DirtyPackagesPath;
+	FEditorFileUtils::GetDirtyPackages(DirtyPackages, FEditorFileUtils::FShouldIgnorePackage::Default);
+
+	Algo::Transform(DirtyPackages, DirtyPackagesPath, UDirtyFilesChangelistValidator::GetPackagePath);
+
+	// Check if any of changelist's files is unsaved
+	for (const FSourceControlStateRef& FileState : FileStates)
+	{
+		if (DirtyPackagesPath.Contains(FileState->GetFilename()))
+		{
+			AssetFails(InAsset, LOCTEXT("DirtyFilesFound", "This changelist contains unsaved modifications. Please save to proceed."), ValidationErrors);
+			return EDataValidationResult::NotValidated;
+		}
+	}
+
+	AssetPasses(InAsset);
+	return EDataValidationResult::Valid;
+}
+
+#undef LOCTEXT_NAMESPACE
diff --git a/Engine/Plugins/Editor/DataValidation/Source/DataValidation/Public/DirtyFilesChangelistValidator.h b/Engine/Plugins/Editor/DataValidation/Source/DataValidation/Public/DirtyFilesChangelistValidator.h
new file mode 100644
--- /dev/null
+++ b/Engine/Plugins/Editor/DataValidation/Source/DataValidation/Public/DirtyFilesChangelistValidator.h
@@ -0,0 +1,24 @@
+// Copyright Epic Games, Inc. All Rights Reserved.
+
+#pragma once
+
+#include "EditorValidatorBase.h"
+#include "DirtyFilesChangelistValidator.generated.h"
+
+class UPackage;
+
+/**
+* Validates there is no unsaved files in the changelist about to be submitted.
+*/
+UCLASS()
+class DATAVALIDATION_API UDirtyFilesChangelistValidator : public UEditorValidatorBase
+{
+	GENERATED_BODY()
+
+protected:
+	virtual bool CanValidateAsset_Implementation(UObject* InAsset) const override;
+	virtual EDataValidationResult ValidateLoadedAsset_Implementation(UObject* InAsset, TArray<FText>& ValidationErrors) override;
+
+private:
+	FString static GetPackagePath(const UPackage* InPackage);
+};
