#UE Put editor validator loads in a editor-only serialization scope so the assets they reference are not cooked when stripping editor-only data

diff --git a/Engine/Plugins/Editor/DataValidation/Source/DataValidation/Private/EditorValidatorSubsystem.cpp b/Engine/Plugins/Editor/DataValidation/Source/DataValidation/Private/EditorValidatorSubsystem.cpp
--- a/Engine/Plugins/Editor/DataValidation/Source/DataValidation/Private/EditorValidatorSubsystem.cpp
+++ b/Engine/Plugins/Editor/DataValidation/Source/DataValidation/Private/EditorValidatorSubsystem.cpp
@@ -110,6 +110,7 @@ void UEditorValidatorSubsystem::RegisterBlueprintValidators()
 			UObject* ValidatorObject = BPAssetData.ToSoftObjectPath().ResolveObject();
 			if (ValidatorObject == nullptr)
 			{
+				FSoftObjectPathSerializationScope SerializationScope(NAME_None, NAME_None, ESoftObjectPathCollectType::EditorOnlyCollect, ESoftObjectPathSerializeType::AlwaysSerialize);
 				ValidatorObject = BPAssetData.ToSoftObjectPath().TryLoad();
 			}
 			if (ValidatorObject)
