* Adds Usecases to Editor Validators, allowing validators to specify for which usecase they should be executed by override
* Integrates TheCoalitation suggestion about FValidateAssetsResult struct, giving detailed information about the validated assets and results

Unshelved from pending changelist '17324854', review 17395959

#changelist validated

diff --git a/Engine/Plugins/Editor/DataValidation/Source/DataValidation/Private/DataValidationCommandlet.cpp b/Engine/Plugins/Editor/DataValidation/Source/DataValidation/Private/DataValidationCommandlet.cpp
--- a/Engine/Plugins/Editor/DataValidation/Source/DataValidation/Private/DataValidationCommandlet.cpp
+++ b/Engine/Plugins/Editor/DataValidation/Source/DataValidation/Private/DataValidationCommandlet.cpp
@@ -1,6 +1,7 @@
 // Copyright Epic Games, Inc. All Rights Reserved.
 
 #include "DataValidationCommandlet.h"
+#include "DataValidationModule.h"
 #include "IAssetRegistry.h"
 #include "AssetRegistryHelpers.h"
 #include "AssetRegistryModule.h"
@@ -43,7 +44,14 @@ bool UDataValidationCommandlet::ValidateData()
 	UEditorValidatorSubsystem* EditorValidationSubsystem = GEditor->GetEditorSubsystem<UEditorValidatorSubsystem>();
 	check(EditorValidationSubsystem);
 
-	EditorValidationSubsystem->ValidateAssets(AssetDataList);
+	FValidateAssetsSettings Settings;
+	FValidateAssetsResults Results;
+
+	Settings.bSkipExcludedDirectories = true;
+	Settings.bShowIfNoFailures = true;
+	Settings.ValidationUsecase = EDataValidationUsecase::Commandlet;
+
+	EditorValidationSubsystem->ValidateAssetsWithSettings(AssetDataList, Settings, Results);
 
 	return true;
 }
diff --git a/Engine/Plugins/Editor/DataValidation/Source/DataValidation/Private/DataValidationModule.cpp b/Engine/Plugins/Editor/DataValidation/Source/DataValidation/Private/DataValidationModule.cpp
--- a/Engine/Plugins/Editor/DataValidation/Source/DataValidation/Private/DataValidationModule.cpp
+++ b/Engine/Plugins/Editor/DataValidation/Source/DataValidation/Private/DataValidationModule.cpp
@@ -36,9 +36,9 @@ class FDataValidationModule : public IDataValidationModule
 	// End IModuleInterface
 
 	/** Validates selected assets and opens a window to report the results. If bValidateDependencies is true it will also validate any assets that the selected assets depend on. */
-	virtual void ValidateAssets(TArray<FAssetData> SelectedAssets, bool bValidateDependencies) override;
+	virtual void ValidateAssets(const TArray<FAssetData>& SelectedAssets, bool bValidateDependencies, const EDataValidationUsecase InValidationUsecase) override;
 
-	void ValidateFolders(TArray<FString> SelectedFolders);
+	void ValidateFolders(TArray<FString> SelectedFolders, const EDataValidationUsecase InValidationUsecase);
 
 private:
 	TSharedRef<FExtender> OnExtendContentBrowserAssetSelectionMenu(const TArray<FAssetData>& SelectedAssets);
@@ -190,29 +190,35 @@ void FDataValidationModule::FindAssetDependencies(const FAssetRegistryModule& As
 	}
 }
 
-void FDataValidationModule::ValidateAssets(TArray<FAssetData> SelectedAssets, bool bValidateDependencies)
+void FDataValidationModule::ValidateAssets(const TArray<FAssetData>& SelectedAssets, bool bValidateDependencies, const EDataValidationUsecase InValidationUsecase)
 {
 	FAssetRegistryModule& AssetRegistryModule = FModuleManager::LoadModuleChecked<FAssetRegistryModule>(TEXT("AssetRegistry"));
 
 	TSet<FAssetData> DependentAssets;
+
 	if (bValidateDependencies)
 	{
 		for (const FAssetData& Asset : SelectedAssets)
 		{
 			FindAssetDependencies(AssetRegistryModule, Asset, DependentAssets);
 		}
-
-		SelectedAssets = DependentAssets.Array();
 	}
 
 	UEditorValidatorSubsystem* EditorValidationSubsystem = GEditor->GetEditorSubsystem<UEditorValidatorSubsystem>();
 	if (EditorValidationSubsystem)
 	{
-		EditorValidationSubsystem->ValidateAssets(SelectedAssets, false);
+		FValidateAssetsSettings Settings;
+		FValidateAssetsResults Results;
+
+		Settings.bSkipExcludedDirectories = false;
+		Settings.bShowIfNoFailures = true;
+		Settings.ValidationUsecase = InValidationUsecase;
+
+		EditorValidationSubsystem->ValidateAssetsWithSettings(bValidateDependencies ? DependentAssets.Array() : SelectedAssets, Settings, Results);
 	}
 }
 
-void FDataValidationModule::ValidateFolders(TArray<FString> SelectedFolders)
+void FDataValidationModule::ValidateFolders(TArray<FString> SelectedFolders, const EDataValidationUsecase InValidationUsecase)
 {
 	FAssetRegistryModule& AssetRegistryModule = FModuleManager::Get().LoadModuleChecked<FAssetRegistryModule>(TEXT("AssetRegistry"));
 
@@ -228,7 +234,7 @@ void FDataValidationModule::ValidateFolders(TArray<FString> SelectedFolders)
 	TArray<FAssetData> AssetList;
 	AssetRegistryModule.Get().GetAssets(Filter, AssetList);
 
-	ValidateAssets(AssetList, false);
+	ValidateAssets(AssetList, false, InValidationUsecase);
 }
 
 // Extend content browser menu for groups of selected assets
@@ -253,14 +259,15 @@ void FDataValidationModule::CreateDataValidationContentBrowserAssetMenu(FMenuBui
 		LOCTEXT("ValidateAssetsTabTitle", "Validate Assets"),
 		LOCTEXT("ValidateAssetsTooltipText", "Runs data validation on these assets."),
 		FSlateIcon(),
-		FUIAction(FExecuteAction::CreateRaw(this, &FDataValidationModule::ValidateAssets, SelectedAssets, false))
+		FUIAction(FExecuteAction::CreateLambda([this, SelectedAssets]() { ValidateAssets(SelectedAssets, false, EDataValidationUsecase::Manual); }))
 	);
+
 	MenuBuilder.AddMenuEntry
 	(
 		LOCTEXT("ValidateAssetsAndDependenciesTabTitle", "Validate Assets and Dependencies"),
 		LOCTEXT("ValidateAssetsAndDependenciesTooltipText", "Runs data validation on these assets and all assets they depend on."),
 		FSlateIcon(),
-		FUIAction(FExecuteAction::CreateRaw(this, &FDataValidationModule::ValidateAssets, SelectedAssets, true))
+		FUIAction(FExecuteAction::CreateLambda([this, SelectedAssets]() { ValidateAssets(SelectedAssets, true, EDataValidationUsecase::Manual); }))
 	);
 }
 
@@ -302,7 +309,7 @@ void FDataValidationModule::CreateDataValidationContentBrowserPathMenu(FMenuBuil
 			const EAppReturnType::Type Result = FMessageDialog::Open(EAppMsgType::YesNo, FText::Format(LOCTEXT("DataValidationConfirmation", "Are you sure you want to proceed with validating the following folders?\n\n{Paths}"), Args));
 			if (Result == EAppReturnType::Yes)
 			{
-				ValidateFolders(SelectedPaths);
+				ValidateFolders(SelectedPaths, EDataValidationUsecase::Manual);
 			}
 		}))
 	);
diff --git a/Engine/Plugins/Editor/DataValidation/Source/DataValidation/Private/EditorValidatorBase.cpp b/Engine/Plugins/Editor/DataValidation/Source/DataValidation/Private/EditorValidatorBase.cpp
--- a/Engine/Plugins/Editor/DataValidation/Source/DataValidation/Private/EditorValidatorBase.cpp
+++ b/Engine/Plugins/Editor/DataValidation/Source/DataValidation/Private/EditorValidatorBase.cpp
@@ -10,6 +10,7 @@
 #include "Logging/MessageLog.h"
 #include "Misc/UObjectToken.h"
 #include "Logging/TokenizedMessage.h"
+#include "DataValidationModule.h"
 
 #define LOCTEXT_NAMESPACE "AssetValidation"
 
@@ -19,6 +20,11 @@ UEditorValidatorBase::UEditorValidatorBase()
 	bIsEnabled = true;
 }
 
+bool UEditorValidatorBase::CanValidate_Implementation(const EDataValidationUsecase InUsecase) const
+{
+	return true;
+}
+
 bool UEditorValidatorBase::CanValidateAsset_Implementation(UObject* InAsset) const
 {
 	return false;
diff --git a/Engine/Plugins/Editor/DataValidation/Source/DataValidation/Private/EditorValidatorSubsystem.cpp b/Engine/Plugins/Editor/DataValidation/Source/DataValidation/Private/EditorValidatorSubsystem.cpp
--- a/Engine/Plugins/Editor/DataValidation/Source/DataValidation/Private/EditorValidatorSubsystem.cpp
+++ b/Engine/Plugins/Editor/DataValidation/Source/DataValidation/Private/EditorValidatorSubsystem.cpp
@@ -18,6 +18,7 @@
 #include "AssetData.h"
 #include "ISourceControlModule.h"
 #include "DataValidationChangelist.h"
+#include "DataValidationModule.h"
 
 #define LOCTEXT_NAMESPACE "EditorValidationSubsystem"
 
@@ -146,7 +147,7 @@ void UEditorValidatorSubsystem::CleanupValidators()
 	Validators.Empty();
 }
 
-EDataValidationResult UEditorValidatorSubsystem::IsObjectValid(UObject* InObject, TArray<FText>& ValidationErrors, TArray<FText>& ValidationWarnings) const
+EDataValidationResult UEditorValidatorSubsystem::IsObjectValid(UObject* InObject, TArray<FText>& ValidationErrors, TArray<FText>& ValidationWarnings, const EDataValidationUsecase InValidationUsecase) const
 {
 	EDataValidationResult Result = EDataValidationResult::NotValidated;
 	
@@ -159,7 +160,7 @@ EDataValidationResult UEditorValidatorSubsystem::IsObjectValid(UObject* InObject
 		{
 			for (auto ValidatorPair : Validators)
 			{
-				if (ValidatorPair.Value && ValidatorPair.Value->IsEnabled() && ValidatorPair.Value->CanValidateAsset(InObject))
+				if (ValidatorPair.Value && ValidatorPair.Value->IsEnabled() && ValidatorPair.Value->CanValidate(InValidationUsecase) && ValidatorPair.Value->CanValidateAsset(InObject))
 				{
 					ValidatorPair.Value->ResetValidationState();
 					EDataValidationResult NewResult = ValidatorPair.Value->ValidateLoadedAsset(InObject, ValidationErrors);
@@ -181,14 +182,14 @@ EDataValidationResult UEditorValidatorSubsystem::IsObjectValid(UObject* InObject
 	return Result;
 }
 
-EDataValidationResult UEditorValidatorSubsystem::IsAssetValid(FAssetData& AssetData, TArray<FText>& ValidationErrors, TArray<FText>& ValidationWarnings) const
+EDataValidationResult UEditorValidatorSubsystem::IsAssetValid(const FAssetData& AssetData, TArray<FText>& ValidationErrors, TArray<FText>& ValidationWarnings, const EDataValidationUsecase InValidationUsecase) const
 {
 	if (AssetData.IsValid())
 	{
 		UObject* Obj = AssetData.GetAsset();
 		if (Obj)
 		{
-			return IsObjectValid(Obj, ValidationErrors, ValidationWarnings);
+			return IsObjectValid(Obj, ValidationErrors, ValidationWarnings, InValidationUsecase);
 		}
 		return EDataValidationResult::NotValidated;
 	}
@@ -197,10 +198,22 @@ EDataValidationResult UEditorValidatorSubsystem::IsAssetValid(FAssetData& AssetD
 }
 
 int32 UEditorValidatorSubsystem::ValidateAssets(TArray<FAssetData> AssetDataList, bool bSkipExcludedDirectories, bool bShowIfNoFailures) const
+{
+	FValidateAssetsSettings Settings;
+	FValidateAssetsResults Results;
+
+	Settings.bSkipExcludedDirectories = bSkipExcludedDirectories;
+	Settings.bShowIfNoFailures = bShowIfNoFailures;
+	
+	return ValidateAssetsWithSettings(AssetDataList, Settings, Results);
+}
+
+int32 UEditorValidatorSubsystem::ValidateAssetsWithSettings(const TArray<FAssetData>& AssetDataList, const FValidateAssetsSettings& InSettings, FValidateAssetsResults& OutResults) const
 {
 	FScopedSlowTask SlowTask(1.0f, LOCTEXT("ValidatingDataTask", "Validating Data..."));
-	SlowTask.Visibility = bShowIfNoFailures ? ESlowTaskVisibility::ForceVisible : ESlowTaskVisibility::Invisible;
-	if (bShowIfNoFailures)
+	SlowTask.Visibility = InSettings.bShowIfNoFailures ? ESlowTaskVisibility::ForceVisible : ESlowTaskVisibility::Invisible;
+	
+	if (InSettings.bShowIfNoFailures)
 	{
 		SlowTask.MakeDialogDelayed(.1f);
 	}
@@ -247,13 +260,13 @@ int32 UEditorValidatorSubsystem::ValidateAssets(TArray<FAssetData> AssetDataList
 	};
 
 	// Now add to map or update as needed
-	for (FAssetData& Data : AssetDataList)
+	for (const FAssetData& Data : AssetDataList)
 	{
 		FText ValidatingMessage = FText::Format(LOCTEXT("ValidatingFilename", "Validating {0}"), FText::FromString(Data.GetFullName()));
 		SlowTask.EnterProgressFrame(1.0f / NumFilesToValidate, ValidatingMessage);
 
 		// Check exclusion path
-		if (bSkipExcludedDirectories && IsPathExcludedFromValidation(Data.PackageName.ToString()))
+		if (InSettings.bSkipExcludedDirectories && IsPathExcludedFromValidation(Data.PackageName.ToString()))
 		{
 			++NumFilesSkipped;
 			continue;
@@ -263,7 +276,7 @@ int32 UEditorValidatorSubsystem::ValidateAssets(TArray<FAssetData> AssetDataList
 
 		TArray<FText> ValidationErrors;
 		TArray<FText> ValidationWarnings;
-		EDataValidationResult Result = IsAssetValid(Data, ValidationErrors, ValidationWarnings);
+		EDataValidationResult Result = IsAssetValid(Data, ValidationErrors, ValidationWarnings, InSettings.ValidationUsecase);
 		++NumFilesChecked;
 
 		AddAssetLogTokens(EMessageSeverity::Error, ValidationErrors, Data.PackageName);
@@ -294,7 +307,7 @@ int32 UEditorValidatorSubsystem::ValidateAssets(TArray<FAssetData> AssetDataList
 			}
 			else if (Result == EDataValidationResult::NotValidated)
 			{
-				if (bShowIfNoFailures)
+				if (InSettings.bShowIfNoFailures)
 				{
 					DataValidationLog.Info()->AddToken(FAssetNameToken::Create(Data.PackageName.ToString()))
 						->AddToken(FTextToken::Create(LOCTEXT("NotValidatedDataResult", "has no data validation.")));
@@ -307,7 +320,14 @@ int32 UEditorValidatorSubsystem::ValidateAssets(TArray<FAssetData> AssetDataList
 	const bool bFailed = (NumInvalidFiles > 0);
 	const bool bAtLeastOneWarning = (NumFilesWithWarnings > 0);
 
-	if (bFailed || bAtLeastOneWarning || bShowIfNoFailures)
+	OutResults.NumChecked = NumFilesChecked;
+	OutResults.NumValid = NumValidFiles;
+	OutResults.NumInvalid = NumInvalidFiles;
+	OutResults.NumSkipped = NumFilesSkipped;
+	OutResults.NumWarnings = NumFilesWithWarnings;
+	OutResults.NumUnableToValidate = NumFilesUnableToValidate;
+
+	if (bFailed || bAtLeastOneWarning || InSettings.bShowIfNoFailures)
 	{
 		FFormatNamedArguments Arguments;
 		Arguments.Add(TEXT("Result"), bFailed ? LOCTEXT("Failed", "FAILED") : LOCTEXT("Succeeded", "SUCCEEDED"));
@@ -347,7 +367,15 @@ void UEditorValidatorSubsystem::ValidateOnSave(TArray<FAssetData> AssetDataList,
 	FMessageLog DataValidationLog("AssetCheck");
 	FText SavedAsset = AssetDataList.Num() == 1 ? FText::FromName(AssetDataList[0].AssetName) : LOCTEXT("MultipleErrors", "multiple assets");
 	DataValidationLog.NewPage(FText::Format(LOCTEXT("DataValidationLogPage", "Asset Save: {0}"), SavedAsset));
-	if (ValidateAssets(MoveTemp(AssetDataList), true, false) > 0)
+
+	FValidateAssetsSettings Settings;
+	FValidateAssetsResults Results;
+
+	Settings.bSkipExcludedDirectories = true;
+	Settings.bShowIfNoFailures = false;
+	Settings.ValidationUsecase = EDataValidationUsecase::Save;
+
+	if (ValidateAssetsWithSettings(AssetDataList, Settings, Results) > 0)
 	{
 		const FText ErrorMessageNotification = FText::Format(
 			LOCTEXT("ValidationFailureNotification", "Validation failed when saving {0}, check Data Validation log"), SavedAsset);
@@ -438,7 +466,7 @@ void UEditorValidatorSubsystem::ValidateChangelistPreSubmit(FSourceControlChange
 	Changelist->AddToRoot();
 
 	SlowTask.EnterProgressFrame(1.0f, LOCTEXT("ValidatingChangelist", "Validating changelist to submit"));
-	OutResult = IsObjectValid(Changelist, OutValidationErrors, OutValidationWarnings);
+	OutResult = IsObjectValid(Changelist, OutValidationErrors, OutValidationWarnings, EDataValidationUsecase::PreSubmit);
 
 	Changelist->RemoveFromRoot();
 }
diff --git a/Engine/Plugins/Editor/DataValidation/Source/DataValidation/Public/DataValidationModule.h b/Engine/Plugins/Editor/DataValidation/Source/DataValidation/Public/DataValidationModule.h
--- a/Engine/Plugins/Editor/DataValidation/Source/DataValidation/Public/DataValidationModule.h
+++ b/Engine/Plugins/Editor/DataValidation/Source/DataValidation/Public/DataValidationModule.h
@@ -8,6 +8,28 @@
 #include "Modules/ModuleInterface.h"
 #include "Modules/ModuleManager.h"
 
+UENUM(BlueprintType)
+enum class EDataValidationUsecase : uint8
+{
+	/** No usecase specified */
+	None = 0,
+
+	/** Triggered on user's demand */
+	Manual,
+
+	/** A commandlet invoked the validation */
+	Commandlet,
+
+	/** Saving a package triggered the validation */
+	Save,
+
+	/** Submit dialog triggered the validation */
+	PreSubmit,
+
+	/** Triggered by blueprint or c++ */
+	Script,
+};
+
 /**
  * The public interface to this module
  */
@@ -41,5 +63,5 @@ public:
 	}
 
 	/** Validates selected assets and opens a window to report the results. If bValidateDependencies is true it will also validate any assets that the selected assets depend on. */
-	virtual void ValidateAssets(TArray<FAssetData> SelectedAssets, bool bValidateDependencies) = 0;
+	virtual void ValidateAssets(const TArray<FAssetData>& SelectedAssets, bool bValidateDependencies, const EDataValidationUsecase InValidationUsecase) = 0;
 };
diff --git a/Engine/Plugins/Editor/DataValidation/Source/DataValidation/Public/EditorValidatorBase.h b/Engine/Plugins/Editor/DataValidation/Source/DataValidation/Public/EditorValidatorBase.h
--- a/Engine/Plugins/Editor/DataValidation/Source/DataValidation/Public/EditorValidatorBase.h
+++ b/Engine/Plugins/Editor/DataValidation/Source/DataValidation/Public/EditorValidatorBase.h
@@ -12,6 +12,8 @@
 #include "Logging/LogVerbosity.h"
 #include "Internationalization/Text.h"
 #include "UObject/TextProperty.h"
+#include "DataValidationModule.h"
+
 #include "EditorValidatorBase.generated.h"
 
 /*
@@ -31,6 +33,10 @@ class DATAVALIDATION_API UEditorValidatorBase : public UObject
 public:
 	UEditorValidatorBase();
 
+	/** Override this to determine whether or not you can use this validator given this usecase */
+	UFUNCTION(BlueprintNativeEvent, Category = "Asset Validation")
+	bool CanValidate(const EDataValidationUsecase InUsecase) const;
+
 	/** Override this to determine whether or not you can validate a given asset with this validator */
 	UFUNCTION(BlueprintNativeEvent, Category = "Asset Validation")
 	bool CanValidateAsset(UObject* InAsset) const;
diff --git a/Engine/Plugins/Editor/DataValidation/Source/DataValidation/Public/EditorValidatorSubsystem.h b/Engine/Plugins/Editor/DataValidation/Source/DataValidation/Public/EditorValidatorSubsystem.h
--- a/Engine/Plugins/Editor/DataValidation/Source/DataValidation/Public/EditorValidatorSubsystem.h
+++ b/Engine/Plugins/Editor/DataValidation/Source/DataValidation/Public/EditorValidatorSubsystem.h
@@ -13,6 +13,7 @@
 #include "Engine/EngineTypes.h"
 #include "AssetData.h"
 #include "Logging/LogMacros.h"
+#include "DataValidationModule.h"
 
 #include "EditorValidatorSubsystem.generated.h"
 
@@ -22,6 +23,58 @@ typedef TSharedPtr<class ISourceControlChangelist, ESPMode::ThreadSafe> FSourceC
 
 DECLARE_LOG_CATEGORY_EXTERN(LogContentValidation, Log, All);
 
+USTRUCT(BlueprintType)
+struct DATAVALIDATION_API FValidateAssetsResults
+{
+	GENERATED_BODY()
+
+	FValidateAssetsResults() = default;
+
+	/** Amount of tested assets */
+	UPROPERTY(VisibleAnywhere, BlueprintReadOnly, Category = "Validation")
+	int NumChecked = 0;
+
+	/** Amount of assets without errors or warnings */
+	UPROPERTY(VisibleAnywhere, BlueprintReadOnly, Category = "Validation")
+	int NumValid = 0;
+
+	/** Amount of assets with errors */
+	UPROPERTY(VisibleAnywhere, BlueprintReadOnly, Category = "Validation")
+	int NumInvalid = 0;
+
+	/** Amount of assets skipped */
+	UPROPERTY(VisibleAnywhere, BlueprintReadOnly, Category = "Validation")
+	int NumSkipped = 0;
+
+	/** Amount of assets with warnings */
+	UPROPERTY(VisibleAnywhere, BlueprintReadOnly, Category = "Validation")
+	int NumWarnings = 0;
+
+	/** Amount of assets that could not be validated */
+	UPROPERTY(VisibleAnywhere, BlueprintReadOnly, Category = "Validation")
+	int NumUnableToValidate = 0;
+};
+
+USTRUCT(BlueprintType)
+struct DATAVALIDATION_API FValidateAssetsSettings
+{
+	GENERATED_BODY()
+
+	FValidateAssetsSettings() = default;
+
+	/** If true, will not validate files in excluded directories */
+	UPROPERTY(EditAnywhere, BlueprintReadWrite, Category = "Validation")
+	bool bSkipExcludedDirectories = true;
+
+	/** If true, will add notifications for files with no validation and display even if everything passes */
+	UPROPERTY(EditAnywhere, BlueprintReadWrite, Category = "Validation")
+	bool bShowIfNoFailures = true;
+
+	/** The usecase requiring datavalidation */
+	UPROPERTY(EditAnywhere, BlueprintReadWrite, Category = "Validation")
+	EDataValidationUsecase ValidationUsecase = EDataValidationUsecase::None;
+};
+
 /**
 * Implements the settings for Data Validation 
 */
@@ -70,23 +123,27 @@ public:
 	 * @return Returns Valid if the object contains valid data; returns Invalid if the object contains invalid data; returns NotValidated if no validations was performed on the object
 	 */
 	UFUNCTION(BlueprintCallable, Category = "Asset Validation")
-	virtual EDataValidationResult IsObjectValid(UObject* InObject, TArray<FText>& ValidationErrors, TArray<FText>& ValidationWarnings) const;
+	virtual EDataValidationResult IsObjectValid(UObject* InObject, TArray<FText>& ValidationErrors, TArray<FText>& ValidationWarnings, const EDataValidationUsecase InValidationUsecase) const;
 
 	/**
 	 * @return Returns Valid if the object pointed to by AssetData contains valid data; returns Invalid if the object contains invalid data or does not exist; returns NotValidated if no validations was performed on the object
 	 */
 	UFUNCTION(BlueprintCallable, Category = "Asset Validation")
-	virtual EDataValidationResult IsAssetValid(FAssetData& AssetData, TArray<FText>& ValidationErrors, TArray<FText>& ValidationWarnings) const;
+	virtual EDataValidationResult IsAssetValid(const FAssetData& AssetData, TArray<FText>& ValidationErrors, TArray<FText>& ValidationWarnings, const EDataValidationUsecase InValidationUsecase) const;
 
 	/**
 	 * Called to validate assets from either the UI or a commandlet
-	 * @param bSkipExcludedDirectories If true, will not validate files in excluded directories
-	 * @param bShowIfNoFailures If true, will add notifications for files with no validation and display even if everything passes
+	 * @param InSettings Structure passing context and settings for ValidateAssetsWithSettings
+	 * @param OutResults More detailed information about the results of the validate assets command
 	 * @returns Number of assets with validation failures or warnings
 	 */
 	UFUNCTION(BlueprintCallable, Category = "Asset Validation")
+	virtual int32 ValidateAssetsWithSettings(const TArray<FAssetData>& AssetDataList, const FValidateAssetsSettings& InSettings, FValidateAssetsResults& OutResults) const;
+	
+	UE_DEPRECATED(5.0, "Use ValidateAssetsWithSettings instead")
+	UFUNCTION(BlueprintCallable, Category = "Asset Validation", meta = (DeprecatedFunction, DeprecationMessage = "Use ValidateAssetsWithSettings instead"))
 	virtual int32 ValidateAssets(TArray<FAssetData> AssetDataList, bool bSkipExcludedDirectories = true, bool bShowIfNoFailures = true) const;
-
+	
 	UE_DEPRECATED(5.0, "Use version that takes bProceduralSave instead")
 	virtual void ValidateOnSave(TArray<FAssetData> AssetDataList) const;
 
