-WorldPartitionChangeListValidator
-When validating a CL with actors in a WorldPartition map, it locates and uses the ActorDescContainer for the map to do a validation of the Actors in the CL and inform the user of various errors. The same validation code is used in MapCheck.

diff --git a/Engine/Plugins/Editor/DataValidation/Source/DataValidation/Private/WorldPartitionChangelistValidator.cpp b/Engine/Plugins/Editor/DataValidation/Source/DataValidation/Private/WorldPartitionChangelistValidator.cpp
new file mode 100644
--- /dev/null
+++ b/Engine/Plugins/Editor/DataValidation/Source/DataValidation/Private/WorldPartitionChangelistValidator.cpp
@@ -0,0 +1,250 @@
+// Copyright Epic Games, Inc. All Rights Reserved.
+
+#include "WorldPartitionChangelistValidator.h"
+#include "DataValidationChangelist.h"
+
+#include "AssetData.h"
+#include "AssetRegistryModule.h"
+#include "ISourceControlProvider.h"
+#include "ISourceControlModule.h"
+#include "SourceControlHelpers.h"
+
+#include "Engine/World.h"
+#include "WorldPartition/WorldPartitionActorDesc.h"
+#include "WorldPartition/ActorDescContainer.h"
+#include "WorldPartition/WorldPartition.h"
+
+#define LOCTEXT_NAMESPACE "WorldPartitionChangelistValidation"
+
+bool UWorldPartitionChangelistValidator::CanValidateAsset_Implementation(UObject* InAsset) const
+{
+	return (InAsset != nullptr) && (UDataValidationChangelist::StaticClass() == InAsset->GetClass());
+}
+
+EDataValidationResult UWorldPartitionChangelistValidator::ValidateLoadedAsset_Implementation(UObject* InAsset, TArray<FText>& ValidationErrors)
+{
+	TRACE_CPUPROFILER_EVENT_SCOPE(UWorldPartitionChangelistValidator::ValidateLoadedAsset_Implementation);
+
+	UDataValidationChangelist* ChangeList = CastChecked<UDataValidationChangelist>(InAsset);
+	
+	Errors = &ValidationErrors;
+	
+	EDataValidationResult Result = ValidateActorsListFromChangeList(ChangeList);
+
+	if (Result == EDataValidationResult::Invalid)
+	{
+		AssetFails(InAsset, LOCTEXT("WorldPartitionValidationFail", "This changelist contains modifications that aren't valid at the World Partition level. Please see Source Control log and correct the errors."), ValidationErrors);
+	}
+	else
+	{
+		AssetPasses(InAsset);
+	}
+		
+	return Result;
+}
+
+FString GetPrettyPackageName(const FWorldPartitionActorDescView& Desc)
+{
+	FString AssetPath = Desc.GetActorPath().ToString();
+	
+	int32 LastDot = -1;
+	if (AssetPath.FindLastChar('.', LastDot))
+	{
+		AssetPath.LeftInline(LastDot);
+	}
+
+	FString AssetName = Desc.GetActorLabel().ToString();
+	
+	if (AssetName.Len() == 0)
+	{
+		AssetName = Desc.GetActorName().ToString();
+	}
+	
+	return AssetPath + TEXT(".") + AssetName;
+}
+
+// Extract all Actors/Map from Changelist (in OFPA this should be one Actor per Package, and we'll discard all Actors from non WorldPartition maps)
+// and add them to a Map of World->Files[] so that we can do one validation per world. Once Worlds are identified, we either the UActorDescContainer 
+// from memory (if loaded) or request it to be loaded, we then build a Set of objects that interest us from the Actors in the CL 
+EDataValidationResult UWorldPartitionChangelistValidator::ValidateActorsListFromChangeList(UDataValidationChangelist* Changelist)
+{
+	ISourceControlProvider& SourceControlProvider = ISourceControlModule::Get().GetProvider();
+	FSourceControlChangelistStatePtr ChangelistState = SourceControlProvider.GetState(Changelist->Changelist->AsShared(), EStateCacheUsage::Use);
+		 	
+	// Figure out which world(s) those assets are in and split the files per world
+	TMap<FName, TArray<FAssetData>>	MapToFiles;
+
+	for (const FSourceControlStateRef& File : ChangelistState->GetFilesStates())
+	{		
+		// Skip deleted files since we're not validating references in this validator 
+		if (File->IsDeleted())
+		{
+			continue;
+		}
+
+		FString PackageName;
+		if (FPackageName::TryConvertFilenameToLongPackageName(File->GetFilename(), PackageName))
+		{						
+			// Extract Actor & Map from package
+			TArray<FAssetData> PackageAssetsData;
+			USourceControlHelpers::GetAssetDataFromPackage(PackageName, PackageAssetsData);
+			
+			for (FAssetData& AssetData : PackageAssetsData)
+			{
+				// Check that the asset is an actor
+				if (AssetData.GetClass()->IsChildOf<AActor>())
+				{
+					// WorldPartition actors are all in OFPA mode so they're external
+					// Extract the MapNMame from the ObjectPath (<PathToPackage>.<mapName>:<level>.<actorName>)
+					FSoftObjectPath ActorPath = FSoftObjectPath(AssetData.ObjectPath);
+					FName MapAssetName = ActorPath.GetAssetPathName();
+
+					TArray<FAssetData>* Files = MapToFiles.Find(MapAssetName);
+
+					if (!Files)
+					{
+						if (ULevel::GetIsLevelPartitionedFromPackage(ActorPath.GetLongPackageFName()))
+						{
+							Files = &MapToFiles.Add(MapAssetName);
+						}
+					}
+
+					if (Files)	// A null Files indicates a World not using World Partition and OFPA 
+					{
+						Files->Add(AssetData);
+					}
+				}
+			}
+		}
+	}
+
+	// For Each world 
+	for (TTuple<FName, TArray<FAssetData>>& It : MapToFiles)
+	{
+		const FName& MapName = It.Get<0>();
+		const TArray<FAssetData>& ActorsData = It.Get<1>();
+
+		// Find/Load the ActorDescContainer
+		UWorld* World = FindObject<UWorld>(nullptr, *MapName.ToString(), true);
+
+		UActorDescContainer* ActorDescContainer = nullptr;
+		
+		if (World != nullptr)
+		{
+			// World is Loaded reuse the ActorDescContainer of the World
+			ActorDescContainer = World->GetWorldPartition();
+		}
+		else
+		{
+			// Find in memory 
+			ActorDescContainer = NewObject<UActorDescContainer>();
+			ActorDescContainer->Initialize(nullptr, MapName);
+		}
+
+		// Build a set of Relevant Actor Guids to scope error messages to what's contained in the CL 
+		RelevantActorGuids.Reset();
+
+		for (const FAssetData& ActorData : ActorsData)
+		{
+			// Get the FWorldPartitionActor			
+			const FWorldPartitionActorDesc* ActorDesc = ActorDescContainer->GetActorDesc(ActorData.AssetName.ToString());
+
+			if (ActorDesc != nullptr)
+			{
+				RelevantActorGuids.Add(ActorDesc->GetGuid());
+			}
+		}
+
+		// Invoke static WorldPartition Validation from the ActorDescContainer
+		UWorldPartition::CheckForErrors(this, ActorDescContainer);		
+	}
+
+	if (Errors->Num())
+	{
+		return EDataValidationResult::Invalid;
+	}
+		
+	return EDataValidationResult::Valid;
+	
+}
+
+bool UWorldPartitionChangelistValidator::Filter(const FWorldPartitionActorDescView& ActorDescView)
+{
+	if (RelevantActorGuids.Find(ActorDescView.GetGuid()))
+	{
+		return true;
+	}
+
+	return false;
+}
+
+void UWorldPartitionChangelistValidator::OnInvalidReference(const FWorldPartitionActorDescView& ActorDescView, const FGuid& ReferenceGuid)
+{
+	if (Filter(ActorDescView))
+	{
+		FText CurrentError = FText::Format(LOCTEXT("DataValidation.Changelist.WorldPartition.InvalidReference", "Actor {0} has a missing reference to %s"),
+											FText::FromString(GetPrettyPackageName(ActorDescView)), 
+											FText::FromString(ReferenceGuid.ToString()));
+
+		Errors->Add(CurrentError);
+	}
+
+}
+
+void UWorldPartitionChangelistValidator::OnInvalidReferenceGridPlacement(const FWorldPartitionActorDescView& ActorDescView, const FWorldPartitionActorDescView& ReferenceActorDescView)
+{
+	if (Filter(ActorDescView) || Filter(ReferenceActorDescView))
+	{
+		const FString StreamedActor(TEXT("Streamed actor"));
+		const FString AlwaysLoadedActor(TEXT("Always loaded actor"));
+
+		const bool bIsActorDescAlwaysLoaded = ActorDescView.GetGridPlacement() == EActorGridPlacement::AlwaysLoaded;
+		const bool bIsActorDescRefAlwaysLoaded = ReferenceActorDescView.GetGridPlacement() == EActorGridPlacement::AlwaysLoaded;
+
+		FText CurrentError = FText::Format(LOCTEXT("DataValidation.Changelist.WorldPartition.InvalidReferenceGridPlacement", "{0} {1} is referencing {2} {3}."),
+											FText::FromString(bIsActorDescAlwaysLoaded ? *AlwaysLoadedActor : *StreamedActor),
+											FText::FromString(GetPrettyPackageName(ActorDescView)),
+											FText::FromString(bIsActorDescRefAlwaysLoaded ? *AlwaysLoadedActor : *StreamedActor),
+											FText::FromString(GetPrettyPackageName(ReferenceActorDescView)));
+
+		Errors->Add(CurrentError);
+	}
+}
+
+void UWorldPartitionChangelistValidator::OnInvalidReferenceDataLayers(const FWorldPartitionActorDescView& ActorDescView, const FWorldPartitionActorDescView& ReferenceActorDescView)
+{	
+	if (Filter(ActorDescView) || Filter(ReferenceActorDescView))
+	{
+		FText CurrentError = FText::Format(LOCTEXT("DataValidation.Changelist.WorldPartition.DataLayerError", "{0} is referencing {1} but both Actors are using a different set of DataLayers."),
+											FText::FromString(GetPrettyPackageName(ActorDescView)),
+											FText::FromString(GetPrettyPackageName(ReferenceActorDescView)));
+
+		Errors->Add(CurrentError);
+	}
+}
+
+void UWorldPartitionChangelistValidator::OnInvalidReferenceLevelScriptStreamed(const FWorldPartitionActorDescView& ActorDescView)
+{
+	if (Filter(ActorDescView))
+	{		
+		FText CurrentError = FText::Format(LOCTEXT("DataValidation.Changelist.WorldPartition.InvalidReferenceLevelScriptStreamed", "Level Script Blueprint references streamed actor {0}."),
+											FText::FromString(GetPrettyPackageName(ActorDescView)));
+		
+		Errors->Add(CurrentError);
+	}
+}
+
+void UWorldPartitionChangelistValidator::OnInvalidReferenceLevelScriptDataLayers(const FWorldPartitionActorDescView& ActorDescView)
+{
+	if (Filter(ActorDescView))
+	{
+		FText CurrentError = FText::Format(LOCTEXT("DataValidation.Changelist.WorldPartition.InvalidReferenceLevelScriptDataLayers", "Level Script Blueprint references streamed actor {0} with a non empty set of data layers"),
+											FText::FromString(GetPrettyPackageName(ActorDescView)));
+
+		Errors->Add(CurrentError);
+	}
+}
+
+
+
+#undef LOCTEXT_NAMESPACE
\ No newline at end of file
diff --git a/Engine/Plugins/Editor/DataValidation/Source/DataValidation/Private/WorldPartitionChangelistValidator.h b/Engine/Plugins/Editor/DataValidation/Source/DataValidation/Private/WorldPartitionChangelistValidator.h
new file mode 100644
--- /dev/null
+++ b/Engine/Plugins/Editor/DataValidation/Source/DataValidation/Private/WorldPartitionChangelistValidator.h
@@ -0,0 +1,38 @@
+// Copyright Epic Games, Inc. All Rights Reserved.
+#pragma once
+
+#include "EditorValidatorBase.h"
+#include "WorldPartition/ErrorHandling/WorldPartitionStreamingGenerationErrorHandler.h"
+
+#include "WorldPartitionChangelistValidator.generated.h"
+
+class UDataValidationChangelist;
+class UActorDescContainer;
+class FWorldPartitionActorDesc;
+
+UCLASS()
+class DATAVALIDATION_API UWorldPartitionChangelistValidator : public UEditorValidatorBase, public IStreamingGenerationErrorHandler
+{
+	GENERATED_BODY()
+
+protected:		
+	
+	TArray<FText>*							Errors = nullptr;
+	TSet<FGuid>								RelevantActorGuids;
+
+	virtual bool CanValidateAsset_Implementation(UObject* InAsset) const override;
+	virtual EDataValidationResult ValidateLoadedAsset_Implementation(UObject* InAsset, TArray<FText>& ValidationErrors) override;
+
+	EDataValidationResult ValidateActorsListFromChangeList(UDataValidationChangelist* Changelist);
+
+	// Return true if this ActorDescView is pertinent to the current ChangeList
+	bool Filter(const FWorldPartitionActorDescView& ActorDescView);
+
+	// IStreamingGenerationErrorHandler Interface methods
+	virtual void OnInvalidReference(const FWorldPartitionActorDescView& ActorDescView, const FGuid& ReferenceGuid) override;
+	virtual void OnInvalidReferenceGridPlacement(const FWorldPartitionActorDescView& ActorDescView, const FWorldPartitionActorDescView& ReferenceActorDescView) override;
+	virtual void OnInvalidReferenceDataLayers(const FWorldPartitionActorDescView& ActorDescView, const FWorldPartitionActorDescView& ReferenceActorDescView) override;
+	virtual void OnInvalidReferenceLevelScriptStreamed(const FWorldPartitionActorDescView& ActorDescView) override;
+	virtual void OnInvalidReferenceLevelScriptDataLayers(const FWorldPartitionActorDescView& ActorDescView) override;
+};
+
diff --git a/Engine/Source/Runtime/Engine/Private/WorldPartition/ActorDescContainer.cpp b/Engine/Source/Runtime/Engine/Private/WorldPartition/ActorDescContainer.cpp
--- a/Engine/Source/Runtime/Engine/Private/WorldPartition/ActorDescContainer.cpp
+++ b/Engine/Source/Runtime/Engine/Private/WorldPartition/ActorDescContainer.cpp
@@ -297,7 +297,7 @@ void UActorDescContainer::UnpinActor(const FGuid& ActorGuid)
 
 void UActorDescContainer::RegisterEditorDelegates()
 {
-	if (GEditor && !IsTemplate() && !World->IsGameWorld())
+	if (GEditor && !IsTemplate() && World && !World->IsGameWorld())
 	{
 		FCoreUObjectDelegates::OnObjectPreSave.AddUObject(this, &UActorDescContainer::OnObjectPreSave);
 		FEditorDelegates::OnPackageDeleted.AddUObject(this, &UActorDescContainer::OnPackageDeleted);
@@ -307,7 +307,7 @@ void UActorDescContainer::RegisterEditorDelegates()
 
 void UActorDescContainer::UnregisterEditorDelegates()
 {
-	if (GEditor && !IsTemplate() && !World->IsGameWorld())
+	if (GEditor && !IsTemplate() && World && !World->IsGameWorld())
 	{
 		FCoreUObjectDelegates::OnObjectPreSave.RemoveAll(this);
 		FEditorDelegates::OnPackageDeleted.RemoveAll(this);
diff --git a/Engine/Source/Runtime/Engine/Private/WorldPartition/WorldPartitionStreamingGeneration.cpp b/Engine/Source/Runtime/Engine/Private/WorldPartition/WorldPartitionStreamingGeneration.cpp
--- a/Engine/Source/Runtime/Engine/Private/WorldPartition/WorldPartitionStreamingGeneration.cpp
+++ b/Engine/Source/Runtime/Engine/Private/WorldPartition/WorldPartitionStreamingGeneration.cpp
@@ -167,7 +167,7 @@ class FWorldPartitionStreamingGenerator
 			const FActorContainerID& ContainerID = It.Key();
 			FContainerDescriptor& ContainerDescriptor = It.Value();
 
-			if (ContainerID.IsMainContainer())
+			if (ContainerID.IsMainContainer() && ContainerDescriptor.Container->GetWorld() != nullptr)
 			{
 				// Gather all references to external actors from the level script and make them always loaded
 				if (ULevelScriptBlueprint* LevelScriptBlueprint = ContainerDescriptor.Container->GetWorld()->PersistentLevel->GetLevelScriptBlueprint(true))
@@ -202,7 +202,10 @@ class FWorldPartitionStreamingGenerator
 			}
 
 			// Give the associated runtime hash the possibility to adjust actor descriptor views based on its internal settings, etc.
-			RuntimeHash->UpdateActorDescViewMap(ContainerDescriptor.ActorDescViewMap);
+			if (RuntimeHash != nullptr)
+			{
+				RuntimeHash->UpdateActorDescViewMap(ContainerDescriptor.ActorDescViewMap);
+			}
 
 			// Perform various adjustements based on validations and errors
 			for (auto ActorDescIt = ContainerDescriptor.ActorDescViewMap.CreateIterator(); ActorDescIt; ++ActorDescIt)
@@ -506,6 +509,16 @@ void UWorldPartition::CheckForErrors(IStreamingGenerationErrorHandler* ErrorHand
 		StreamingGenerator.PreparationPhase(this);
 	}
 }
+
+void UWorldPartition::CheckForErrors(IStreamingGenerationErrorHandler* ErrorHandler, const UActorDescContainer* ActorDescContainer) 
+{
+	FActorClusterContext ActorClusterContext;
+	{
+		FWorldPartitionStreamingGenerator StreamingGenerator(nullptr, nullptr, ErrorHandler);
+		StreamingGenerator.PreparationPhase(ActorDescContainer);
+	}
+}
+
 #endif // WITH_EDITOR
 
 #undef LOCTEXT_NAMESPACE
diff --git a/Engine/Source/Runtime/Engine/Public/WorldPartition/WorldPartition.h b/Engine/Source/Runtime/Engine/Public/WorldPartition/WorldPartition.h
--- a/Engine/Source/Runtime/Engine/Public/WorldPartition/WorldPartition.h
+++ b/Engine/Source/Runtime/Engine/Public/WorldPartition/WorldPartition.h
@@ -143,6 +143,7 @@ public:
 	void DumpActorDescs(const FString& Path);
 
 	void CheckForErrors(IStreamingGenerationErrorHandler* ErrorHandler) const;
+	static void CheckForErrors(IStreamingGenerationErrorHandler* ErrorHandler, const UActorDescContainer* ActorDescContainer);
 
 	uint32 GetWantedEditorCellSize() const;
 	void SetEditorWantedCellSize(uint32 InCellSize);
