Fixed duplicate localization key in data validator subsystem.
#rnx

diff --git a/Engine/Plugins/Editor/DataValidation/Source/DataValidation/Private/EditorValidatorSubsystem.cpp b/Engine/Plugins/Editor/DataValidation/Source/DataValidation/Private/EditorValidatorSubsystem.cpp
--- a/Engine/Plugins/Editor/DataValidation/Source/DataValidation/Private/EditorValidatorSubsystem.cpp
+++ b/Engine/Plugins/Editor/DataValidation/Source/DataValidation/Private/EditorValidatorSubsystem.cpp
@@ -454,7 +454,7 @@ void UEditorValidatorSubsystem::ValidateAllSavedPackages()
 
 void UEditorValidatorSubsystem::ValidateChangelistPreSubmit(FSourceControlChangelistPtr InChangelist, EDataValidationResult& OutResult, TArray<FText>& OutValidationErrors, TArray<FText>& OutValidationWarnings) const
 {
-	FScopedSlowTask SlowTask(1.0f, LOCTEXT("ValidatingDataTask", "Validating changelist..."));
+	FScopedSlowTask SlowTask(1.0f, LOCTEXT("ValidatingChangelistTask", "Validating changelist..."));
 	SlowTask.Visibility = ESlowTaskVisibility::Invisible;
 	SlowTask.MakeDialogDelayed(.1f);
 
