Hotfix: Remove ExternalActors & ExternalObjects from Folder Validation.
They cannot be loaded on the fly and spam the validation log with false positives.
UE-144978: A better alternative using WorldParitionErrorHandlers should used to perform the validation

diff --git a/Engine/Plugins/Editor/DataValidation/Source/DataValidation/Private/DataValidationModule.cpp b/Engine/Plugins/Editor/DataValidation/Source/DataValidation/Private/DataValidationModule.cpp
--- a/Engine/Plugins/Editor/DataValidation/Source/DataValidation/Private/DataValidationModule.cpp
+++ b/Engine/Plugins/Editor/DataValidation/Source/DataValidation/Private/DataValidationModule.cpp
@@ -25,6 +25,7 @@
 #include "WorkspaceMenuStructureModule.h"
 #include "EditorValidatorSubsystem.h"
 #include "ISettingsModule.h"
+#include "Algo/RemoveIf.h"
 
 #define LOCTEXT_NAMESPACE "DataValidationModule"
 
@@ -234,6 +235,19 @@ void FDataValidationModule::ValidateFolders(TArray<FString> SelectedFolders, con
 	TArray<FAssetData> AssetList;
 	AssetRegistryModule.Get().GetAssets(Filter, AssetList);
 
+	// UE-144978 : Remove ExternalActors & ExternalObjects from assets to be validated.
+	// If ExternalActors are not loaded, they will spam the validation log as they can't
+	// be loaded on the fly like other assets.
+	auto IsAssetPackageExternal = [](const FAssetData& AssetData)
+	{
+		FString ObjectPath = AssetData.ObjectPath.ToString();
+		FStringView ClassName, PackageName, ObjectName, SubObjectName;
+		FPackageName::SplitFullObjectPath(FStringView(ObjectPath), ClassName, PackageName, ObjectName, SubObjectName);
+
+		return FName(PackageName) != AssetData.PackageName;
+	};
+	AssetList.SetNum(Algo::RemoveIf(AssetList, IsAssetPackageExternal));
+
 	ValidateAssets(AssetList, false, InValidationUsecase);
 }
 
