Add a parameter to DataValidation commandlet (e.g. AssetType=PhysicsAsset) to specify a subset of types to test.

diff --git a/Engine/Plugins/Editor/DataValidation/Source/DataValidation/Private/DataValidationCommandlet.cpp b/Engine/Plugins/Editor/DataValidation/Source/DataValidation/Private/DataValidationCommandlet.cpp
--- a/Engine/Plugins/Editor/DataValidation/Source/DataValidation/Private/DataValidationCommandlet.cpp
+++ b/Engine/Plugins/Editor/DataValidation/Source/DataValidation/Private/DataValidationCommandlet.cpp
@@ -16,13 +16,9 @@ int32 UDataValidationCommandlet::Main(const FString& FullCommandLine)
 {
 	UE_LOG(LogDataValidation, Log, TEXT("--------------------------------------------------------------------------------------------"));
 	UE_LOG(LogDataValidation, Log, TEXT("Running DataValidation Commandlet"));
-	TArray<FString> Tokens;
-	TArray<FString> Switches;
-	TMap<FString, FString> Params;
-	ParseCommandLine(*FullCommandLine, Tokens, Switches, Params);
 
 	// validate data
-	if (!ValidateData())
+	if (!ValidateData(FullCommandLine))
 	{
 		UE_LOG(LogDataValidation, Warning, TEXT("Errors occurred while validating data"));
 		return 2; // return something other than 1 for error since the engine will return 1 if any other system (possibly unrelated) logged errors during execution.
@@ -34,12 +30,26 @@ int32 UDataValidationCommandlet::Main(const FString& FullCommandLine)
 }
 
 //static
-bool UDataValidationCommandlet::ValidateData()
+bool UDataValidationCommandlet::ValidateData(const FString& FullCommandLine)
 {
-	FAssetRegistryModule& AssetRegistryModule = FModuleManager::LoadModuleChecked<FAssetRegistryModule>(AssetRegistryConstants::ModuleName);
+	TArray<FString> Tokens;
+	TArray<FString> Switches;
+	TMap<FString, FString> Params;
+	ParseCommandLine(*FullCommandLine, Tokens, Switches, Params);
 
+	FAssetRegistryModule& AssetRegistryModule = FModuleManager::LoadModuleChecked<FAssetRegistryModule>(AssetRegistryConstants::ModuleName);
 	TArray<FAssetData> AssetDataList;
-	AssetRegistryModule.Get().GetAllAssets(AssetDataList);
+	FString AssetTypeString;
+	if (FParse::Value(*FullCommandLine, TEXT("AssetType="), AssetTypeString) && !AssetTypeString.IsEmpty())
+	{
+		const FName AssetType = *AssetTypeString;
+		const bool bSearchSubClasses = true;
+		AssetRegistryModule.Get().GetAssetsByClass(AssetType, AssetDataList, bSearchSubClasses);
+	}
+	else
+	{
+		AssetRegistryModule.Get().GetAllAssets(AssetDataList);
+	}
 
 	UEditorValidatorSubsystem* EditorValidationSubsystem = GEditor->GetEditorSubsystem<UEditorValidatorSubsystem>();
 	check(EditorValidationSubsystem);
diff --git a/Engine/Plugins/Editor/DataValidation/Source/DataValidation/Private/DataValidationModule.cpp b/Engine/Plugins/Editor/DataValidation/Source/DataValidation/Private/DataValidationModule.cpp
--- a/Engine/Plugins/Editor/DataValidation/Source/DataValidation/Private/DataValidationModule.cpp
+++ b/Engine/Plugins/Editor/DataValidation/Source/DataValidation/Private/DataValidationModule.cpp
@@ -147,7 +147,7 @@ void FDataValidationModule::Menu_ValidateData()
 	}
 
 	// validate the data
-	bool bSuccess = UDataValidationCommandlet::ValidateData();
+	bool bSuccess = UDataValidationCommandlet::ValidateData(FString());
 
 	// display an error if the task failed
 	if (!bSuccess)
diff --git a/Engine/Plugins/Editor/DataValidation/Source/DataValidation/Public/DataValidationCommandlet.h b/Engine/Plugins/Editor/DataValidation/Source/DataValidation/Public/DataValidationCommandlet.h
--- a/Engine/Plugins/Editor/DataValidation/Source/DataValidation/Public/DataValidationCommandlet.h
+++ b/Engine/Plugins/Editor/DataValidation/Source/DataValidation/Public/DataValidationCommandlet.h
@@ -18,9 +18,9 @@ public:
 	}
 
 	// Begin UCommandlet Interface
-	virtual int32 Main(const FString& Params) override;
+	virtual int32 Main(const FString& FullCommandLine) override;
 	// End UCommandlet Interface
 
 	// do the validation without creating a commandlet
-	static bool ValidateData();
+	static bool ValidateData(const FString& FullCommandLine);
 };
