World Partition Validation: better wording for actors referencing other actors in a different set of runtime data layers.

#rnx

diff --git a/Engine/Plugins/Editor/DataValidation/Source/DataValidation/Private/WorldPartitionChangelistValidator.cpp b/Engine/Plugins/Editor/DataValidation/Source/DataValidation/Private/WorldPartitionChangelistValidator.cpp
--- a/Engine/Plugins/Editor/DataValidation/Source/DataValidation/Private/WorldPartitionChangelistValidator.cpp
+++ b/Engine/Plugins/Editor/DataValidation/Source/DataValidation/Private/WorldPartitionChangelistValidator.cpp
@@ -35,7 +35,7 @@ EDataValidationResult UWorldPartitionChangelistValidator::ValidateLoadedAsset_Im
 
 	if (Result == EDataValidationResult::Invalid)
 	{
-		AssetFails(InAsset, LOCTEXT("WorldPartitionValidationFail", "This changelist contains modifications that aren't valid at the World Partition level. Please see Source Control log and correct the errors."), ValidationErrors);
+		AssetFails(InAsset, LOCTEXT("WorldPartitionValidationFail", "This changelist contains modifications that aren't valid at the world partition level. Please see source control log and correct the errors."), ValidationErrors);
 	}
 	else
 	{
@@ -262,7 +262,7 @@ void UWorldPartitionChangelistValidator::OnInvalidReferenceDataLayers(const FWor
 {	
 	if (Filter(ActorDescView) || Filter(ReferenceActorDescView))
 	{
-		FText CurrentError = FText::Format(LOCTEXT("DataValidation.Changelist.WorldPartition.DataLayerError", "{0} is referencing {1} but both Actors are using a different set of DataLayers."),
+		FText CurrentError = FText::Format(LOCTEXT("DataValidation.Changelist.WorldPartition.DataLayerError", "{0} is referencing {1} but both actors are using a different set of runtime data layers."),
 											FText::FromString(GetPrettyPackageName(ActorDescView)),
 											FText::FromString(GetPrettyPackageName(ReferenceActorDescView)));
 
@@ -274,7 +274,7 @@ void UWorldPartitionChangelistValidator::OnInvalidReferenceRuntimeGrid(const FWo
 {
 	if (Filter(ActorDescView) || Filter(ReferenceActorDescView))
 	{
-		FText CurrentError = FText::Format(LOCTEXT("DataValidation.Changelist.WorldPartition.RuntimeGridError", "{0} is referencing {1} but both Actors are using a different Runtime Grid."),
+		FText CurrentError = FText::Format(LOCTEXT("DataValidation.Changelist.WorldPartition.RuntimeGridError", "{0} is referencing {1} but both actors are using a different runtime grid."),
 			FText::FromString(GetPrettyPackageName(ActorDescView)),
 			FText::FromString(GetPrettyPackageName(ReferenceActorDescView)));
 
@@ -286,7 +286,7 @@ void UWorldPartitionChangelistValidator::OnInvalidReferenceLevelScriptStreamed(c
 {
 	if (Filter(ActorDescView))
 	{		
-		FText CurrentError = FText::Format(LOCTEXT("DataValidation.Changelist.WorldPartition.InvalidReferenceLevelScriptStreamed", "Level Script Blueprint references streamed actor {0}."),
+		FText CurrentError = FText::Format(LOCTEXT("DataValidation.Changelist.WorldPartition.InvalidReferenceLevelScriptStreamed", "Level script blueprint references streamed actor {0}."),
 											FText::FromString(GetPrettyPackageName(ActorDescView)));
 		
 		Errors->Add(CurrentError);
@@ -297,7 +297,7 @@ void UWorldPartitionChangelistValidator::OnInvalidReferenceLevelScriptDataLayers
 {
 	if (Filter(ActorDescView))
 	{
-		FText CurrentError = FText::Format(LOCTEXT("DataValidation.Changelist.WorldPartition.InvalidReferenceLevelScriptDataLayers", "Level Script Blueprint references streamed actor {0} with a non empty set of data layers"),
+		FText CurrentError = FText::Format(LOCTEXT("DataValidation.Changelist.WorldPartition.InvalidReferenceLevelScriptDataLayers", "Level script blueprint references streamed actor {0} with a non empty set of data layers."),
 											FText::FromString(GetPrettyPackageName(ActorDescView)));
 
 		Errors->Add(CurrentError);
@@ -308,7 +308,7 @@ void UWorldPartitionChangelistValidator::OnInvalidReferenceDataLayerAsset(const
 {
 	if (SubmittingWorldDataLayers)
 	{
-		FText CurrentError = FText::Format(LOCTEXT("DataValidation.Changelist.WorldPartition.InvalidDataLayerAsset", "Data Layer {0} has no Data Layer Asset"),
+		FText CurrentError = FText::Format(LOCTEXT("DataValidation.Changelist.WorldPartition.InvalidDataLayerAsset", "Data layer {0} has no data layer asset."),
 			FText::FromName(DataLayerInstance->GetDataLayerFName()));
 
 		Errors->Add(CurrentError);
@@ -321,7 +321,7 @@ void UWorldPartitionChangelistValidator::OnDataLayerHierarchyTypeMismatch(const
 		|| Filter(Parent)
 		|| SubmittingWorldDataLayers)
 	{
-		FText CurrentError = FText::Format(LOCTEXT("DataValidation.Changelist.WorldPartition.DataLayerHierarchyTypeMismatch", "Data Layer {0} is of Type {1} and its parent {2} is of type {3}"),
+		FText CurrentError = FText::Format(LOCTEXT("DataValidation.Changelist.WorldPartition.DataLayerHierarchyTypeMismatch", "Data layer {0} is of type {1} and its parent {2} is of type {3}."),
 			FText::FromString(DataLayerInstance->GetDataLayerFullName()),
 			UEnum::GetDisplayValueAsText(DataLayerInstance->GetType()),
 			FText::FromString(Parent->GetDataLayerFullName()),
@@ -337,7 +337,7 @@ void UWorldPartitionChangelistValidator::OnDataLayerAssetConflict(const UDataLay
 		|| Filter(ConflictingDataLayerInstance)
 		|| SubmittingWorldDataLayers)
 	{
-		FText CurrentError = FText::Format(LOCTEXT("DataValidation.Changelist.WorldPartition.DataLayerAssetConflict", "Data Layer Instance {0} and Data Layer Instance {1} are both referencing Data Layer Asset {2}"),
+		FText CurrentError = FText::Format(LOCTEXT("DataValidation.Changelist.WorldPartition.DataLayerAssetConflict", "Data layer instance {0} and data layer instance {1} are both referencing data layer asset {2}."),
 			FText::FromName(DataLayerInstance->GetDataLayerFName()),
 			FText::FromName(ConflictingDataLayerInstance->GetDataLayerFName()),
 			FText::FromString(DataLayerInstance->GetAsset()->GetFullName()));
diff --git a/Engine/Source/Runtime/Engine/Private/WorldPartition/ErrorHandling/WorldPartitionStreamingGenerationLogErrorHandler.cpp b/Engine/Source/Runtime/Engine/Private/WorldPartition/ErrorHandling/WorldPartitionStreamingGenerationLogErrorHandler.cpp
--- a/Engine/Source/Runtime/Engine/Private/WorldPartition/ErrorHandling/WorldPartitionStreamingGenerationLogErrorHandler.cpp
+++ b/Engine/Source/Runtime/Engine/Private/WorldPartition/ErrorHandling/WorldPartitionStreamingGenerationLogErrorHandler.cpp
@@ -21,7 +21,7 @@ void FStreamingGenerationLogErrorHandler::OnInvalidReferenceGridPlacement(const
 
 void FStreamingGenerationLogErrorHandler::OnInvalidReferenceDataLayers(const FWorldPartitionActorDescView& ActorDescView, const FWorldPartitionActorDescView& ReferenceActorDescView)
 {
-	UE_LOG(LogWorldPartition, Log, TEXT("Actor %s references an actor in a different set of data layers %s"), *ActorDescView.GetActorLabelOrName().ToString(), *ReferenceActorDescView.GetActorLabelOrName().ToString());
+	UE_LOG(LogWorldPartition, Log, TEXT("Actor %s references an actor in a different set of runtime data layers %s"), *ActorDescView.GetActorLabelOrName().ToString(), *ReferenceActorDescView.GetActorLabelOrName().ToString());
 }
 
 void FStreamingGenerationLogErrorHandler::OnInvalidReferenceRuntimeGrid(const FWorldPartitionActorDescView& ActorDescView, const FWorldPartitionActorDescView& ReferenceActorDescView)
diff --git a/Engine/Source/Runtime/Engine/Private/WorldPartition/ErrorHandling/WorldPartitionStreamingGenerationTokenizedMessageErrorHandler.cpp b/Engine/Source/Runtime/Engine/Private/WorldPartition/ErrorHandling/WorldPartitionStreamingGenerationTokenizedMessageErrorHandler.cpp
--- a/Engine/Source/Runtime/Engine/Private/WorldPartition/ErrorHandling/WorldPartitionStreamingGenerationTokenizedMessageErrorHandler.cpp
+++ b/Engine/Source/Runtime/Engine/Private/WorldPartition/ErrorHandling/WorldPartitionStreamingGenerationTokenizedMessageErrorHandler.cpp
@@ -45,7 +45,7 @@ void ITokenizedMessageErrorHandler::OnInvalidReferenceDataLayers(const FWorldPar
 	TSharedRef<FTokenizedMessage> Message = FTokenizedMessage::Create(EMessageSeverity::Error);
 	Message->AddToken(FTextToken::Create(LOCTEXT("TokenMessage_WorldPartition_Actor", "Actor")))
 		->AddToken(FActorToken::Create(ActorDescView.GetActorPath().ToString(), ActorDescView.GetGuid(), FText::FromName(ActorDescView.GetActorLabelOrName())))
-		->AddToken(FTextToken::Create(LOCTEXT("TokenMessage_WorldPartition_ReferenceActorInOtherDataLayers", "references an actor in a different set of data layers")))
+		->AddToken(FTextToken::Create(LOCTEXT("TokenMessage_WorldPartition_ReferenceActorInOtherDataLayers", "references an actor in a different set of runtime data layers")))
 		->AddToken(FActorToken::Create(ReferenceActorDescView.GetActorPath().ToString(), ReferenceActorDescView.GetGuid(), FText::FromName(ReferenceActorDescView.GetActorLabelOrName())));
 
 	AddAdditionalNameToken(Message, FName(TEXT("WorldPartition_ActorReferenceActorInAnotherDataLayer_CheckForErrors")));
