EditorValidatorSubsystem:
- ValidateOnSave doesn't load unloaded assets
- Pass in the ULevel::LoadAllExternalObjectsTag to GetAsset so that WP worlds get fully loaded for validation

diff --git a/Engine/Plugins/Editor/DataValidation/Source/DataValidation/Private/EditorValidatorSubsystem.cpp b/Engine/Plugins/Editor/DataValidation/Source/DataValidation/Private/EditorValidatorSubsystem.cpp
--- a/Engine/Plugins/Editor/DataValidation/Source/DataValidation/Private/EditorValidatorSubsystem.cpp
+++ b/Engine/Plugins/Editor/DataValidation/Source/DataValidation/Private/EditorValidatorSubsystem.cpp
@@ -19,6 +19,7 @@
 #include "ISourceControlModule.h"
 #include "DataValidationChangelist.h"
 #include "DataValidationModule.h"
+#include "Engine/Level.h"
 
 #define LOCTEXT_NAMESPACE "EditorValidationSubsystem"
 
@@ -186,7 +187,7 @@ EDataValidationResult UEditorValidatorSubsystem::IsAssetValid(const FAssetData&
 {
 	if (AssetData.IsValid())
 	{
-		UObject* Obj = AssetData.GetAsset();
+		UObject* Obj = AssetData.GetAsset({ ULevel::LoadAllExternalObjectsTag });
 		if (Obj)
 		{
 			return IsObjectValid(Obj, ValidationErrors, ValidationWarnings, InValidationUsecase);
@@ -272,6 +273,13 @@ int32 UEditorValidatorSubsystem::ValidateAssetsWithSettings(const TArray<FAssetD
 			continue;
 		}
 
+		const bool bLoadAsset = false;
+		if (!InSettings.bLoadAssetsForValidation && !Data.FastGetAsset(bLoadAsset))
+		{
+			++NumFilesSkipped;
+			continue;
+		}
+
 		UE_LOG(LogContentValidation, Display, TEXT("%s"), *ValidatingMessage.ToString());
 
 		TArray<FText> ValidationErrors;
@@ -374,6 +382,7 @@ void UEditorValidatorSubsystem::ValidateOnSave(TArray<FAssetData> AssetDataList,
 	Settings.bSkipExcludedDirectories = true;
 	Settings.bShowIfNoFailures = false;
 	Settings.ValidationUsecase = EDataValidationUsecase::Save;
+	Settings.bLoadAssetsForValidation = false;
 
 	if (ValidateAssetsWithSettings(AssetDataList, Settings, Results) > 0)
 	{
diff --git a/Engine/Plugins/Editor/DataValidation/Source/DataValidation/Public/EditorValidatorSubsystem.h b/Engine/Plugins/Editor/DataValidation/Source/DataValidation/Public/EditorValidatorSubsystem.h
--- a/Engine/Plugins/Editor/DataValidation/Source/DataValidation/Public/EditorValidatorSubsystem.h
+++ b/Engine/Plugins/Editor/DataValidation/Source/DataValidation/Public/EditorValidatorSubsystem.h
@@ -73,6 +73,9 @@ struct DATAVALIDATION_API FValidateAssetsSettings
 	/** The usecase requiring datavalidation */
 	UPROPERTY(EditAnywhere, BlueprintReadWrite, Category = "Validation")
 	EDataValidationUsecase ValidationUsecase = EDataValidationUsecase::None;
+
+	/** If false, unloaded assets will get skipped from validation. */
+	bool bLoadAssetsForValidation = true;
 };
 
 /**
