DataValidationChangelist could miss some first level asset dependencies
* If an asset was already added as a dependency, then a later call to GatherDependencies with that asset would not have it's dependencies added
* Added check for missing dependencies not in source control
* Exclude script/memory packages from the processing

diff --git a/Engine/Plugins/Editor/DataValidation/Source/DataValidation/Private/DataValidationChangelist.cpp b/Engine/Plugins/Editor/DataValidation/Source/DataValidation/Private/DataValidationChangelist.cpp
--- a/Engine/Plugins/Editor/DataValidation/Source/DataValidation/Private/DataValidationChangelist.cpp
+++ b/Engine/Plugins/Editor/DataValidation/Source/DataValidation/Private/DataValidationChangelist.cpp
@@ -13,13 +13,8 @@
 
 #define LOCTEXT_NAMESPACE "DataValidationChangelist"
 
-void GatherDependencies(const FName& InPackageName, TSet<FName>& OutDependencies, bool bGetFullDependencies)
+void GatherDependencies(const FName& InPackageName, TSet<FName>& OutDependencies)
 {
-	if (OutDependencies.Contains(InPackageName))
-	{
-		return;
-	}
-
 	OutDependencies.Add(InPackageName);
 
 	TArray<FAssetData> Assets;
@@ -28,11 +23,8 @@ void GatherDependencies(const FName& InPackageName, TSet<FName>& OutDependencies
 
 	for (const FName& PackageDependency : Dependencies)
 	{
-		if (bGetFullDependencies)
-		{
-			GatherDependencies(PackageDependency, OutDependencies, bGetFullDependencies);
-		}
-		else if (!OutDependencies.Contains(PackageDependency))
+		// Exclude script/memory packages
+		if (FPackageName::IsValidLongPackageName(PackageDependency.ToString()))
 		{
 			OutDependencies.Add(PackageDependency);
 		}
@@ -83,8 +75,7 @@ EDataValidationResult UDataValidationChangelist::IsDataValid(TArray<FText>& Vali
 	// Gather dependencies of every file in the changelist
 	TArray<FName> FilesInChangelist;
 	TSet<FName> ExternalDependenciesSet;
-	const bool bGetFullDependencyHierarchy = false;
-
+	
 	for (const FSourceControlStateRef& File : ChangelistState->GetFilesStates())
 	{
 		// We shouldn't consider dependencies of deleted files
@@ -97,7 +88,7 @@ EDataValidationResult UDataValidationChangelist::IsDataValid(TArray<FText>& Vali
 		if (FPackageName::TryConvertFilenameToLongPackageName(File->GetFilename(), PackageName))
 		{
 			FilesInChangelist.Add(*PackageName);
-			GatherDependencies(*PackageName, ExternalDependenciesSet, bGetFullDependencyHierarchy);
+			GatherDependencies(*PackageName, ExternalDependenciesSet);
 		}
 	}
 
@@ -138,13 +129,20 @@ EDataValidationResult UDataValidationChangelist::IsDataValid(TArray<FText>& Vali
 			continue;
 		}
 
-		// If file is checked out, added, it's not in this changelist, which is a problem
+		// Dependency is checked out or added but is not in this changelist
 		if (ExternalDependencyFileState->IsCheckedOut() || ExternalDependencyFileState->IsAdded())
 		{
 			bHasChangelistErrors = true;
 			FText CurrentError = FText::Format(LOCTEXT("DataValidation.Changelist.Error", "{0} is missing from this changelist."), FText::FromString(GetPrettyPackageName(ExternalDependency)));
 			ValidationErrors.Add(CurrentError);
 		}
+		// Dependency is not in source control
+		else if (ExternalDependencyFileState->CanAdd())
+		{
+			bHasChangelistErrors = true;
+			FText CurrentError = FText::Format(LOCTEXT("DataValidation.Changelist.NotInDepot", "{0} is referenced and must also be added to source control '{1}'"), FText::FromString(GetPrettyPackageName(ExternalDependency)), FText::FromString(ExternalPackageFilename));
+			ValidationErrors.Add(CurrentError);
+		}
 	}
 
 	return bHasChangelistErrors ? EDataValidationResult::Invalid : EDataValidationResult::Valid;
