Fixed changelist data validation to check if external dependencies are at latest revision. If not, triggers an error for now as it's difficult to say if the changelist is valid or not.

diff --git a/Engine/Plugins/Editor/DataValidation/Source/DataValidation/Private/DataValidationChangelist.cpp b/Engine/Plugins/Editor/DataValidation/Source/DataValidation/Private/DataValidationChangelist.cpp
--- a/Engine/Plugins/Editor/DataValidation/Source/DataValidation/Private/DataValidationChangelist.cpp
+++ b/Engine/Plugins/Editor/DataValidation/Source/DataValidation/Private/DataValidationChangelist.cpp
@@ -136,6 +136,13 @@ EDataValidationResult UDataValidationChangelist::IsDataValid(TArray<FText>& Vali
 			FText CurrentError = FText::Format(LOCTEXT("DataValidation.Changelist.Error", "{0} is missing from this changelist."), FText::FromString(GetPrettyPackageName(ExternalDependency)));
 			ValidationErrors.Add(CurrentError);
 		}
+		// Dependency is not at the latest revision
+		else if (!ExternalDependencyFileState->IsCurrent())
+		{
+			bHasChangelistErrors = true;
+			FText CurrentError = FText::Format(LOCTEXT("DataValidation.Changelist.NotLatest", "{0} is referenced but is not at the latest revision '{1}'"), FText::FromString(GetPrettyPackageName(ExternalDependency)), FText::FromString(ExternalPackageFilename));
+			ValidationErrors.Add(CurrentError);
+		}
 		// Dependency is not in source control
 		else if (ExternalDependencyFileState->CanAdd())
 		{
