Changelist Validation: properly report the case where a reference is missing from disk.

#rnx

diff --git a/Engine/Plugins/Editor/DataValidation/Source/DataValidation/Private/DataValidationChangelist.cpp b/Engine/Plugins/Editor/DataValidation/Source/DataValidation/Private/DataValidationChangelist.cpp
--- a/Engine/Plugins/Editor/DataValidation/Source/DataValidation/Private/DataValidationChangelist.cpp
+++ b/Engine/Plugins/Editor/DataValidation/Source/DataValidation/Private/DataValidationChangelist.cpp
@@ -146,9 +146,18 @@ EDataValidationResult UDataValidationChangelist::IsDataValid(TArray<FText>& Vali
 		// Dependency is not in source control
 		else if (ExternalDependencyFileState->CanAdd())
 		{
-			bHasChangelistErrors = true;
-			FText CurrentError = FText::Format(LOCTEXT("DataValidation.Changelist.NotInDepot", "{0} is referenced and must also be added to source control '{1}'"), FText::FromString(GetPrettyPackageName(ExternalDependency)), FText::FromString(ExternalPackageFilename));
-			ValidationErrors.Add(CurrentError);
+			if (!FPaths::FileExists(ExternalDependencyFileState->GetFilename()))
+			{
+				bHasChangelistErrors = true;
+				FText CurrentError = FText::Format(LOCTEXT("DataValidation.Changelist.NotInWorkspace", "{0} is referenced and cannot be found in workspace '{1}'"), FText::FromString(GetPrettyPackageName(ExternalDependency)), FText::FromString(ExternalPackageFilename));
+				ValidationErrors.Add(CurrentError);
+			}
+			else
+			{
+				bHasChangelistErrors = true;
+				FText CurrentError = FText::Format(LOCTEXT("DataValidation.Changelist.NotInDepot", "{0} is referenced and must also be added to source control '{1}'"), FText::FromString(GetPrettyPackageName(ExternalDependency)), FText::FromString(ExternalPackageFilename));
+				ValidationErrors.Add(CurrentError);
+			}
 		}
 	}
 
