Changelist Validation: added possibility to output warnings and changed reference not latest revision from error to a warning.

#rnx

diff --git a/Engine/Plugins/Editor/DataValidation/Source/DataValidation/Private/DataValidationChangelist.cpp b/Engine/Plugins/Editor/DataValidation/Source/DataValidation/Private/DataValidationChangelist.cpp
--- a/Engine/Plugins/Editor/DataValidation/Source/DataValidation/Private/DataValidationChangelist.cpp
+++ b/Engine/Plugins/Editor/DataValidation/Source/DataValidation/Private/DataValidationChangelist.cpp
@@ -67,7 +67,7 @@ FString GetPrettyPackageName(const FName& InPackageName)
 	}
 }
 
-EDataValidationResult UDataValidationChangelist::IsDataValid(TArray<FText>& ValidationErrors)
+EDataValidationResult UDataValidationChangelist::IsDataValid(TArray<FText>& ValidationErrors, TArray<FText>& ValidationWarnings)
 {
 	ISourceControlProvider& SourceControlProvider = ISourceControlModule::Get().GetProvider();
 	
@@ -140,9 +140,8 @@ EDataValidationResult UDataValidationChangelist::IsDataValid(TArray<FText>& Vali
 		// Dependency is not at the latest revision
 		else if (!ExternalDependencyFileState->IsCurrent())
 		{
-			bHasChangelistErrors = true;
-			FText CurrentError = FText::Format(LOCTEXT("DataValidation.Changelist.NotLatest", "{0} is referenced but is not at the latest revision '{1}'"), FText::FromString(GetPrettyPackageName(ExternalDependency)), FText::FromString(ExternalPackageFilename));
-			ValidationErrors.Add(CurrentError);
+			FText CurrentWarning = FText::Format(LOCTEXT("DataValidation.Changelist.NotLatest", "{0} is referenced but is not at the latest revision '{1}'"), FText::FromString(GetPrettyPackageName(ExternalDependency)), FText::FromString(ExternalPackageFilename));
+			ValidationWarnings.Add(CurrentWarning);
 		}
 		// Dependency is not in source control
 		else if (ExternalDependencyFileState->CanAdd())
diff --git a/Engine/Plugins/Editor/DataValidation/Source/DataValidation/Private/EditorValidatorSubsystem.cpp b/Engine/Plugins/Editor/DataValidation/Source/DataValidation/Private/EditorValidatorSubsystem.cpp
--- a/Engine/Plugins/Editor/DataValidation/Source/DataValidation/Private/EditorValidatorSubsystem.cpp
+++ b/Engine/Plugins/Editor/DataValidation/Source/DataValidation/Private/EditorValidatorSubsystem.cpp
@@ -155,7 +155,7 @@ EDataValidationResult UEditorValidatorSubsystem::IsObjectValid(UObject* InObject
 	if (ensure(InObject))
 	{
 		// First check the class level validation
-		Result = InObject->IsDataValid(ValidationErrors);
+		Result = InObject->IsDataValid(ValidationErrors, ValidationWarnings);
 		// If the asset is still valid or there wasn't a class-level validation, keep validating with custom validators
 		if (Result != EDataValidationResult::Invalid)
 		{
diff --git a/Engine/Plugins/Editor/DataValidation/Source/DataValidation/Public/DataValidationChangelist.h b/Engine/Plugins/Editor/DataValidation/Source/DataValidation/Public/DataValidationChangelist.h
--- a/Engine/Plugins/Editor/DataValidation/Source/DataValidation/Public/DataValidationChangelist.h
+++ b/Engine/Plugins/Editor/DataValidation/Source/DataValidation/Public/DataValidationChangelist.h
@@ -26,7 +26,7 @@ public:
 	/** Default constructor with nothing */
 	UDataValidationChangelist() = default;
 
-	virtual EDataValidationResult IsDataValid(TArray<FText>& ValidationErrors) override;
+	virtual EDataValidationResult IsDataValid(TArray<FText>& ValidationErrors, TArray<FText>& ValidationWarnings) override;
 
 	/** Initializes internal state so the validation can be done */
 	void Initialize(FSourceControlChangelistPtr InChangelist);
diff --git a/Engine/Source/Runtime/CoreUObject/Private/UObject/Obj.cpp b/Engine/Source/Runtime/CoreUObject/Private/UObject/Obj.cpp
--- a/Engine/Source/Runtime/CoreUObject/Private/UObject/Obj.cpp
+++ b/Engine/Source/Runtime/CoreUObject/Private/UObject/Obj.cpp
@@ -4950,6 +4950,11 @@ EDataValidationResult UObject::IsDataValid(TArray<FText>& ValidationErrors)
 	return EDataValidationResult::NotValidated;
 }
 
+EDataValidationResult UObject::IsDataValid(TArray<FText>& ValidationErrors, TArray<FText>& ValidationWarnings)
+{
+	return IsDataValid(ValidationErrors);
+}
+
 #endif // WITH_EDITOR
 /** IsNameStableForNetworking means an object can be referred to its path name (relative to outer) over the network */
 bool UObject::IsNameStableForNetworking() const
diff --git a/Engine/Source/Runtime/CoreUObject/Public/UObject/Object.h b/Engine/Source/Runtime/CoreUObject/Public/UObject/Object.h
--- a/Engine/Source/Runtime/CoreUObject/Public/UObject/Object.h
+++ b/Engine/Source/Runtime/CoreUObject/Public/UObject/Object.h
@@ -956,10 +956,24 @@ public:
 	virtual void PreDestroyFromReplication();
 
 #if WITH_EDITOR
-	/** 
-	 * @return		Returns Valid if this object has data validation rules set up for it and the data for this object is valid. Returns Invalid if it does not pass the rules. Returns NotValidated if no rules are set for this object.
+	/**
+	 * Generic function to validate objects during changelist validations, etc.
+	 *
+	 * @param	ValidationErrors	the array of validation errors.
+	 * @return Valid if this object has data validation rules set up for it and the data for this object is valid. Returns Invalid if it does not pass 
+	 *         the rules. Returns NotValidated if no rules are set for this object.
 	 */
 	virtual EDataValidationResult IsDataValid(TArray<FText>& ValidationErrors);
+
+	/**
+	 * Generic function to validate objects during changelist validations, etc.
+	 *
+	 * @param	ValidationErrors	the array of validation errors.
+	 * @param	ValidationWarnings	the array of validation warnings.
+	 * @return Valid if this object has data validation rules set up for it and the data for this object is valid. Returns Invalid if it does not pass 
+	 *         the rules. Returns NotValidated if no rules are set for this object.
+	 */
+	virtual EDataValidationResult IsDataValid(TArray<FText>& ValidationErrors, TArray<FText>& ValidationWarnings);
 #endif // WITH_EDITOR
 
 	/*----------------------------------------------------------
