Fix World Partition Changelist Validator not reporting invalid references properly

diff --git a/Engine/Plugins/Editor/DataValidation/Source/DataValidation/Private/WorldPartitionChangelistValidator.cpp b/Engine/Plugins/Editor/DataValidation/Source/DataValidation/Private/WorldPartitionChangelistValidator.cpp
--- a/Engine/Plugins/Editor/DataValidation/Source/DataValidation/Private/WorldPartitionChangelistValidator.cpp
+++ b/Engine/Plugins/Editor/DataValidation/Source/DataValidation/Private/WorldPartitionChangelistValidator.cpp
@@ -234,7 +234,7 @@ void UWorldPartitionChangelistValidator::OnInvalidReference(const FWorldPartitio
 {
 	if (Filter(ActorDescView))
 	{
-		FText CurrentError = FText::Format(LOCTEXT("DataValidation.Changelist.WorldPartition.InvalidReference", "Actor {0} has a missing reference to %s"),
+		FText CurrentError = FText::Format(LOCTEXT("DataValidation.Changelist.WorldPartition.InvalidReference", "Actor {0} has a missing reference to {1}"),
 											FText::FromString(GetPrettyPackageName(ActorDescView)), 
 											FText::FromString(ReferenceGuid.ToString()));
 
