#UE Fix a crash when a validator blueprint's parent class could not be found

diff --git a/Engine/Plugins/Editor/DataValidation/Source/DataValidation/Private/EditorValidatorSubsystem.cpp b/Engine/Plugins/Editor/DataValidation/Source/DataValidation/Private/EditorValidatorSubsystem.cpp
--- a/Engine/Plugins/Editor/DataValidation/Source/DataValidation/Private/EditorValidatorSubsystem.cpp
+++ b/Engine/Plugins/Editor/DataValidation/Source/DataValidation/Private/EditorValidatorSubsystem.cpp
@@ -102,7 +102,7 @@ void UEditorValidatorSubsystem::RegisterBlueprintValidators()
 				UObject* Outer = nullptr;
 				ResolveName(Outer, ParentClassName, false, false);
 				ParentClass = FindObject<UClass>(Outer, *ParentClassName);
-				if (!ParentClass->IsChildOf(UEditorValidatorBase::StaticClass()))
+				if (!ParentClass || !ParentClass->IsChildOf(UEditorValidatorBase::StaticClass()))
 				{
 					continue;
 				}
