World Partition Changelist Validation: make sure to validate the world even if there's no actors in the changelist and report errors to external actors if the world package is part of the changelist.

#rnx

diff --git a/Engine/Plugins/Editor/DataValidation/Source/DataValidation/Private/WorldPartitionChangelistValidator.cpp b/Engine/Plugins/Editor/DataValidation/Source/DataValidation/Private/WorldPartitionChangelistValidator.cpp
--- a/Engine/Plugins/Editor/DataValidation/Source/DataValidation/Private/WorldPartitionChangelistValidator.cpp
+++ b/Engine/Plugins/Editor/DataValidation/Source/DataValidation/Private/WorldPartitionChangelistValidator.cpp
@@ -124,7 +124,7 @@ EDataValidationResult UWorldPartitionChangelistValidator::ValidateActorsAndDataL
 			
 			for (FAssetData& AssetData : PackageAssetsData)
 			{
-				if (AssetData.GetClass() == UDataLayerAsset::StaticClass())
+				if (AssetData.GetClass()->IsChildOf<UDataLayerAsset>())
 				{
 					TArray<FName> ReferencerNames;
 					AssetRegistry.GetReferencers(AssetData.PackageName, ReferencerNames, UE::AssetRegistry::EDependencyCategory::All);
@@ -144,13 +144,20 @@ EDataValidationResult UWorldPartitionChangelistValidator::ValidateActorsAndDataL
 			
 					RelevantDataLayerAssets.Add(AssetData.PackageName.ToString());
 				}
-				else if (AssetData.GetClass() == AWorldDataLayers::StaticClass())
+				else if (AssetData.GetClass()->IsChildOf<AWorldDataLayers>())
 				{
 					if (TryAssociateAssetDataToMap(AssetData))
 					{
 						SubmittingWorldDataLayers = true;
 					}
 				}
+				else if (AssetData.GetClass()->IsChildOf<UWorld>())
+				{
+					if (ULevel::GetIsLevelPartitionedFromPackage(*PackageName))
+					{
+						MapToActorsFiles.FindOrAdd(AssetData.ObjectPath);
+					}
+				}
 				else
 				{
 
@@ -186,6 +193,7 @@ EDataValidationResult UWorldPartitionChangelistValidator::ValidateActorsAndDataL
 		}
 
 		// Build a set of Relevant Actor Guids to scope error messages to what's contained in the CL 
+		RelevantMap = MapName;
 		RelevantActorGuids.Reset();
 
 		for (const FAssetData& ActorData : ActorsData)
@@ -221,6 +229,12 @@ bool UWorldPartitionChangelistValidator::Filter(const FWorldPartitionActorDescVi
 		return true;
 	}
 
+	if (!RelevantMap.IsNone())
+	{
+		FSoftObjectPath ActorPath = ActorDescView.GetActorPath().ToString();
+		return ActorPath.GetAssetPathName() == RelevantMap;
+	}
+
 	return false;
 }
 
@@ -353,4 +367,4 @@ void UWorldPartitionChangelistValidator::OnActorNeedsResave(const FWorldPartitio
 	// Changelist validation already ensures that dirty actors must be part of the changelist
 }
 
-#undef LOCTEXT_NAMESPACE
\ No newline at end of file
+#undef LOCTEXT_NAMESPACE
diff --git a/Engine/Plugins/Editor/DataValidation/Source/DataValidation/Private/WorldPartitionChangelistValidator.h b/Engine/Plugins/Editor/DataValidation/Source/DataValidation/Private/WorldPartitionChangelistValidator.h
--- a/Engine/Plugins/Editor/DataValidation/Source/DataValidation/Private/WorldPartitionChangelistValidator.h
+++ b/Engine/Plugins/Editor/DataValidation/Source/DataValidation/Private/WorldPartitionChangelistValidator.h
@@ -18,6 +18,7 @@ class DATAVALIDATION_API UWorldPartitionChangelistValidator : public UEditorVali
 protected:		
 	
 	TArray<FText>*							Errors = nullptr;
+	FName									RelevantMap;
 	TSet<FGuid>								RelevantActorGuids;
 	TSet<FString>							RelevantDataLayerAssets;
 	bool									SubmittingWorldDataLayers = false;
