Removes Validation fail when Uncontrolled Changes are discovered.
Adds light reconcile when reference validation fails.

#changelist validated

diff --git a/Engine/Plugins/Editor/DataValidation/Source/DataValidation/Private/DataValidationChangelist.cpp b/Engine/Plugins/Editor/DataValidation/Source/DataValidation/Private/DataValidationChangelist.cpp
--- a/Engine/Plugins/Editor/DataValidation/Source/DataValidation/Private/DataValidationChangelist.cpp
+++ b/Engine/Plugins/Editor/DataValidation/Source/DataValidation/Private/DataValidationChangelist.cpp
@@ -11,6 +11,7 @@
 #include "ISourceControlModule.h"
 #include "SourceControlHelpers.h"
 #include "SourceControlOperations.h"
+#include "UncontrolledChangelistsModule.h"
 
 #define LOCTEXT_NAMESPACE "DataValidationChangelist"
 
@@ -161,6 +162,11 @@ EDataValidationResult UDataValidationChangelist::IsDataValid(TArray<FText>& Vali
 		}
 	}
 
+	if (bHasChangelistErrors)
+	{
+		FUncontrolledChangelistsModule::Get().OnReconcileAssets();
+	}
+
 	return bHasChangelistErrors ? EDataValidationResult::Invalid : EDataValidationResult::Valid;
 }
 
diff --git a/Engine/Plugins/Editor/DataValidation/Source/DataValidation/Private/UncontrolledChangelistValidator.cpp b/Engine/Plugins/Editor/DataValidation/Source/DataValidation/Private/UncontrolledChangelistValidator.cpp
deleted file mode 100644
diff --git a/Engine/Plugins/Editor/DataValidation/Source/DataValidation/Public/UncontrolledChangelistValidator.h b/Engine/Plugins/Editor/DataValidation/Source/DataValidation/Public/UncontrolledChangelistValidator.h
deleted file mode 100644
