Changelist Validation: fix crash when encountering an invalid asset class.

#rnx

diff --git a/Engine/Plugins/Editor/DataValidation/Source/DataValidation/Private/WorldPartitionChangelistValidator.cpp b/Engine/Plugins/Editor/DataValidation/Source/DataValidation/Private/WorldPartitionChangelistValidator.cpp
--- a/Engine/Plugins/Editor/DataValidation/Source/DataValidation/Private/WorldPartitionChangelistValidator.cpp
+++ b/Engine/Plugins/Editor/DataValidation/Source/DataValidation/Private/WorldPartitionChangelistValidator.cpp
@@ -75,14 +75,14 @@ EDataValidationResult UWorldPartitionChangelistValidator::ValidateActorsAndDataL
 	FSourceControlChangelistStatePtr ChangelistState = SourceControlProvider.GetState(Changelist->Changelist->AsShared(), EStateCacheUsage::Use);
 
 	// Figure out which world(s) those assets are in and split the files per world
-	TMap<FName, TSet<FAssetData>>	MapToActorsFiles;
-	auto TryAssociateAssetDataToMap = [&MapToActorsFiles](const FAssetData& AssetData)
+	TMap<FName, TSet<FAssetData>> MapToActorsFiles;
+	auto TryAssociateActorToMap = [&MapToActorsFiles](const FAssetData& AssetData)
 	{
 		// Check that the asset is an actor
 		if (FWorldPartitionActorDescUtils::IsValidActorDescriptorFromAssetData(AssetData))
 		{
 			// WorldPartition actors are all in OFPA mode so they're external
-			// Extract the MapName from the ObjectPath (<PathToPackage>.<mapName>:<level>.<actorName>)
+			// Extract the MapName from the ObjectPath (<PathToPackage>.<MapName>:<Level>.<ActorName>)
 			FSoftObjectPath ActorPath = FSoftObjectPath(AssetData.ObjectPath);
 			FName MapAssetName = ActorPath.GetAssetPathName();
 
@@ -124,45 +124,44 @@ EDataValidationResult UWorldPartitionChangelistValidator::ValidateActorsAndDataL
 			
 			for (FAssetData& AssetData : PackageAssetsData)
 			{
-				if (AssetData.GetClass()->IsChildOf<UDataLayerAsset>())
-				{
-					TArray<FName> ReferencerNames;
-					AssetRegistry.GetReferencers(AssetData.PackageName, ReferencerNames, UE::AssetRegistry::EDependencyCategory::All);
-			
-					FARFilter Filter;
-					Filter.bIncludeOnlyOnDiskAssets = true;
-					Filter.PackageNames = MoveTemp(ReferencerNames);
-					Filter.ClassPaths.Add(AWorldDataLayers::StaticClass()->GetClassPathName());
+				UClass* AssetClass = AssetData.GetClass();
 
-					TArray<FAssetData> DataLayerReferencers;
-					AssetRegistry.GetAssets(Filter, DataLayerReferencers);
-
-					for (const FAssetData& DataLayerReferencer : DataLayerReferencers)
-					{
-						TryAssociateAssetDataToMap(DataLayerReferencer);
-					}
-			
-					RelevantDataLayerAssets.Add(AssetData.PackageName.ToString());
+				if (TryAssociateActorToMap(AssetData))
+				{
+					check(AssetClass);
+					check(AssetClass->IsChildOf<AActor>())
+					SubmittingWorldDataLayers = AssetClass->IsChildOf<AWorldDataLayers>();
 				}
-				else if (AssetData.GetClass()->IsChildOf<AWorldDataLayers>())
+				else if (AssetClass)
 				{
-					if (TryAssociateAssetDataToMap(AssetData))
+					if (AssetClass->IsChildOf<UDataLayerAsset>())
 					{
-						SubmittingWorldDataLayers = true;
+						TArray<FName> ReferencerNames;
+						AssetRegistry.GetReferencers(AssetData.PackageName, ReferencerNames, UE::AssetRegistry::EDependencyCategory::All);
+			
+						FARFilter Filter;
+						Filter.bIncludeOnlyOnDiskAssets = true;
+						Filter.PackageNames = MoveTemp(ReferencerNames);
+						Filter.ClassPaths.Add(AWorldDataLayers::StaticClass()->GetClassPathName());
+
+						TArray<FAssetData> DataLayerReferencers;
+						AssetRegistry.GetAssets(Filter, DataLayerReferencers);
+
+						for (const FAssetData& DataLayerReferencer : DataLayerReferencers)
+						{
+							TryAssociateActorToMap(DataLayerReferencer);
+						}
+			
+						RelevantDataLayerAssets.Add(AssetData.PackageName.ToString());
 					}
-				}
-				else if (AssetData.GetClass()->IsChildOf<UWorld>())
-				{
-					if (ULevel::GetIsLevelPartitionedFromPackage(*PackageName))
+					else if (AssetClass->IsChildOf<UWorld>())
 					{
-						MapToActorsFiles.FindOrAdd(AssetData.ObjectPath);
+						if (ULevel::GetIsLevelPartitionedFromPackage(*PackageName))
+						{
+							MapToActorsFiles.FindOrAdd(AssetData.ObjectPath);
+						}
 					}
 				}
-				else
-				{
-
-					TryAssociateAssetDataToMap(AssetData);
-				}
 			}
 		}
 	}
