Always show a progres dialog when validating assets to save, even if the requested settings wants to hide the resulting dialog if there's no validation errors, as this process can take some time and result in no errors in the end, showing no progress to the user.

#rnx

diff --git a/Engine/Plugins/Editor/DataValidation/Source/DataValidation/Private/EditorValidatorSubsystem.cpp b/Engine/Plugins/Editor/DataValidation/Source/DataValidation/Private/EditorValidatorSubsystem.cpp
--- a/Engine/Plugins/Editor/DataValidation/Source/DataValidation/Private/EditorValidatorSubsystem.cpp
+++ b/Engine/Plugins/Editor/DataValidation/Source/DataValidation/Private/EditorValidatorSubsystem.cpp
@@ -212,12 +212,8 @@ int32 UEditorValidatorSubsystem::ValidateAssets(TArray<FAssetData> AssetDataList
 int32 UEditorValidatorSubsystem::ValidateAssetsWithSettings(const TArray<FAssetData>& AssetDataList, const FValidateAssetsSettings& InSettings, FValidateAssetsResults& OutResults) const
 {
 	FScopedSlowTask SlowTask(1.0f, LOCTEXT("ValidatingDataTask", "Validating Data..."));
-	SlowTask.Visibility = InSettings.bShowIfNoFailures ? ESlowTaskVisibility::ForceVisible : ESlowTaskVisibility::Invisible;
-	
-	if (InSettings.bShowIfNoFailures)
-	{
-		SlowTask.MakeDialogDelayed(.1f);
-	}
+	SlowTask.Visibility = ESlowTaskVisibility::ForceVisible;
+	SlowTask.MakeDialogDelayed(.1f);
 
 	FMessageLog DataValidationLog("AssetCheck");
 
