FNames containing asset paths are deprecated. FSoftObjectPath or FTopLevelAssetPath should be used instead.
Fixups for data validation.

diff --git a/Engine/Plugins/Editor/DataValidation/Source/DataValidation/Private/DataValidationChangelist.cpp b/Engine/Plugins/Editor/DataValidation/Source/DataValidation/Private/DataValidationChangelist.cpp
--- a/Engine/Plugins/Editor/DataValidation/Source/DataValidation/Private/DataValidationChangelist.cpp
+++ b/Engine/Plugins/Editor/DataValidation/Source/DataValidation/Private/DataValidationChangelist.cpp
@@ -40,7 +40,7 @@ FString GetPrettyPackageName(const FName& InPackageName)
 
 	if (Assets.Num() > 0)
 	{
-		FString AssetPath = Assets[0].ObjectPath.ToString();
+		FString AssetPath = Assets[0].GetObjectPathString();
 
 		int32 LastDot = -1;
 		if (AssetPath.FindLastChar('.', LastDot))
diff --git a/Engine/Plugins/Editor/DataValidation/Source/DataValidation/Private/DataValidationModule.cpp b/Engine/Plugins/Editor/DataValidation/Source/DataValidation/Private/DataValidationModule.cpp
--- a/Engine/Plugins/Editor/DataValidation/Source/DataValidation/Private/DataValidationModule.cpp
+++ b/Engine/Plugins/Editor/DataValidation/Source/DataValidation/Private/DataValidationModule.cpp
@@ -181,8 +181,7 @@ void FDataValidationModule::FindAssetDependencies(const FAssetRegistryModule& As
 					FString DependencyObjectString = FString::Printf(TEXT("%s.%s"), *DependencyPackageString, *FPackageName::GetLongPackageAssetName(DependencyPackageString));
 
 					// recurse on each dependency
-					FName ObjectPath(*DependencyObjectString);
-					FAssetData DependentAsset = AssetRegistryModule.Get().GetAssetByObjectPath(ObjectPath);
+					FAssetData DependentAsset = AssetRegistryModule.Get().GetAssetByObjectPath(FSoftObjectPath(DependencyObjectString));
 
 					FindAssetDependencies(AssetRegistryModule, DependentAsset, DependentAssets);
 				}
@@ -240,7 +239,7 @@ void FDataValidationModule::ValidateFolders(TArray<FString> SelectedFolders, con
 	// be loaded on the fly like other assets.
 	auto IsAssetPackageExternal = [](const FAssetData& AssetData)
 	{
-		FString ObjectPath = AssetData.ObjectPath.ToString();
+		FString ObjectPath = AssetData.GetObjectPathString();
 		FStringView ClassName, PackageName, ObjectName, SubObjectName;
 		FPackageName::SplitFullObjectPath(FStringView(ObjectPath), ClassName, PackageName, ObjectName, SubObjectName);
 
diff --git a/Engine/Plugins/Editor/DataValidation/Source/DataValidation/Private/EditorValidator_Localization.cpp b/Engine/Plugins/Editor/DataValidation/Source/DataValidation/Private/EditorValidator_Localization.cpp
--- a/Engine/Plugins/Editor/DataValidation/Source/DataValidation/Private/EditorValidator_Localization.cpp
+++ b/Engine/Plugins/Editor/DataValidation/Source/DataValidation/Private/EditorValidator_Localization.cpp
@@ -39,7 +39,7 @@ EDataValidationResult UEditorValidator_Localization::ValidateLoadedAsset_Impleme
 		FString SourceObjectPath;
 		if (FPackageLocalizationUtil::ConvertLocalizedToSource(AssetPathName, SourceObjectPath))
 		{
-			const FAssetData SourceAssetData = AssetRegistry.GetAssetByObjectPath(*SourceObjectPath);
+			const FAssetData SourceAssetData = AssetRegistry.GetAssetByObjectPath(FSoftObjectPath(SourceObjectPath));
 
 			// Does this source asset exist?
 			// It is valid to have orphan localized assets, as they may be a direct reference of another localized asset
@@ -87,7 +87,7 @@ EDataValidationResult UEditorValidator_Localization::ValidateLoadedAsset_Impleme
 					FString LocalizedObjectPath;
 					if (FPackageLocalizationUtil::ConvertSourceToLocalized(AssetPathName, CultureName, LocalizedObjectPath))
 					{
-						const FAssetData LocalizedAssetData = AssetRegistry.GetAssetByObjectPath(*LocalizedObjectPath);
+						const FAssetData LocalizedAssetData = AssetRegistry.GetAssetByObjectPath(FSoftObjectPath(LocalizedObjectPath));
 
 						if (LocalizedAssetData.IsValid())
 						{
