Fixed WorldPartition changelist validator crash accessing world partition's ActorDescContainer on a loaded but not initialized world.

diff --git a/Engine/Plugins/Editor/DataValidation/Source/DataValidation/Private/WorldPartitionChangelistValidator.cpp b/Engine/Plugins/Editor/DataValidation/Source/DataValidation/Private/WorldPartitionChangelistValidator.cpp
--- a/Engine/Plugins/Editor/DataValidation/Source/DataValidation/Private/WorldPartitionChangelistValidator.cpp
+++ b/Engine/Plugins/Editor/DataValidation/Source/DataValidation/Private/WorldPartitionChangelistValidator.cpp
@@ -178,7 +178,9 @@ EDataValidationResult UWorldPartitionChangelistValidator::ValidateActorsAndDataL
 			// World is Loaded reuse the ActorDescContainer of the World
 			ActorDescContainer = World->GetWorldPartition()->GetActorDescContainer();
 		}
-		else
+		
+		// Even if world is valid, its world partition is not necessarily initialized
+		if (!ActorDescContainer)
 		{
 			// Find in memory failed, load the ActorDescContainer
 			ActorDescContainer = NewObject<UActorDescContainer>();
