World Partition Validation: make spatially loaded actors referencing non-spatially loaded ones a warning instead of an error since this will be supported in the future.

#rnx

diff --git a/Engine/Plugins/Editor/DataValidation/Source/DataValidation/Private/WorldPartitionChangelistValidator.cpp b/Engine/Plugins/Editor/DataValidation/Source/DataValidation/Private/WorldPartitionChangelistValidator.cpp
--- a/Engine/Plugins/Editor/DataValidation/Source/DataValidation/Private/WorldPartitionChangelistValidator.cpp
+++ b/Engine/Plugins/Editor/DataValidation/Source/DataValidation/Private/WorldPartitionChangelistValidator.cpp
@@ -258,16 +258,20 @@ void UWorldPartitionChangelistValidator::OnInvalidReferenceGridPlacement(const F
 {
 	if (Filter(ActorDescView) || Filter(ReferenceActorDescView))
 	{
-		const FString SpatiallyLoadedActor(TEXT("Spatially loaded actor"));
-		const FString NonSpatiallyLoadedActor(TEXT("Non-spatially loaded loaded actor"));
+		// Only report errors for non-spatially loaded actor referencing a spatially loaded actor
+		if (!ActorDescView.GetIsSpatiallyLoaded())
+		{
+			const FString SpatiallyLoadedActor(TEXT("Spatially loaded actor"));
+			const FString NonSpatiallyLoadedActor(TEXT("Non-spatially loaded loaded actor"));
 
-		FText CurrentError = FText::Format(LOCTEXT("DataValidation.Changelist.WorldPartition.InvalidReferenceGridPlacement", "{0} {1} is referencing {2} {3}."),
-											FText::FromString(ActorDescView.GetIsSpatiallyLoaded() ? *SpatiallyLoadedActor : *NonSpatiallyLoadedActor),
-											FText::FromString(GetPrettyPackageName(ActorDescView)),
-											FText::FromString(ReferenceActorDescView.GetIsSpatiallyLoaded() ? *SpatiallyLoadedActor : *NonSpatiallyLoadedActor),
-											FText::FromString(GetPrettyPackageName(ReferenceActorDescView)));
+			FText CurrentError = FText::Format(LOCTEXT("DataValidation.Changelist.WorldPartition.InvalidReferenceGridPlacement", "{0} {1} is referencing {2} {3}."),
+												FText::FromString(ActorDescView.GetIsSpatiallyLoaded() ? *SpatiallyLoadedActor : *NonSpatiallyLoadedActor),
+												FText::FromString(GetPrettyPackageName(ActorDescView)),
+												FText::FromString(ReferenceActorDescView.GetIsSpatiallyLoaded() ? *SpatiallyLoadedActor : *NonSpatiallyLoadedActor),
+												FText::FromString(GetPrettyPackageName(ReferenceActorDescView)));
 
-		Errors->Add(CurrentError);
+			Errors->Add(CurrentError);
+		}
 	}
 }
 
diff --git a/Engine/Source/Runtime/Engine/Private/WorldPartition/ErrorHandling/WorldPartitionStreamingGenerationTokenizedMessageErrorHandler.cpp b/Engine/Source/Runtime/Engine/Private/WorldPartition/ErrorHandling/WorldPartitionStreamingGenerationTokenizedMessageErrorHandler.cpp
--- a/Engine/Source/Runtime/Engine/Private/WorldPartition/ErrorHandling/WorldPartitionStreamingGenerationTokenizedMessageErrorHandler.cpp
+++ b/Engine/Source/Runtime/Engine/Private/WorldPartition/ErrorHandling/WorldPartitionStreamingGenerationTokenizedMessageErrorHandler.cpp
@@ -27,8 +27,9 @@ void ITokenizedMessageErrorHandler::OnInvalidReferenceGridPlacement(const FWorld
 {
 	const FText SpatiallyLoadedActor(LOCTEXT("TokenMessage_WorldPartition_SpatiallyLoadedActor", "Spatially loaded actor"));
 	const FText NonSpatiallyLoadedActor(LOCTEXT("TokenMessage_WorldPartition_NonSpatiallyLoadedActor", "Non-spatially loaded actor"));
+	const EMessageSeverity::Type MessageSeverity = ActorDescView.GetIsSpatiallyLoaded() ? EMessageSeverity::Warning : EMessageSeverity::Error;
 
-	TSharedRef<FTokenizedMessage> Message = FTokenizedMessage::Create(EMessageSeverity::Error);
+	TSharedRef<FTokenizedMessage> Message = FTokenizedMessage::Create(MessageSeverity);
 	Message->AddToken(FTextToken::Create(ActorDescView.GetIsSpatiallyLoaded() ? SpatiallyLoadedActor : NonSpatiallyLoadedActor))
 		->AddToken(FActorToken::Create(ActorDescView.GetActorSoftPath().ToString(), ActorDescView.GetGuid(), FText::FromString(GetFullActorName(ActorDescView))))
 		->AddToken(FTextToken::Create(LOCTEXT("TokenMessage_WorldPartition_References", "references")))
