Data Validation: added UObject::IsDataValid version taking an FDataValidationContext input and removed newly introduced version that gathers warnings.
No functionnal change, this is mainly to avoid having to introduce new versions of UObject::IsDataValid in the future (old one will be deprecated in 5.2).
Change proposed by nick.darnell.

#rnx

diff --git a/Engine/Plugins/Editor/DataValidation/Source/DataValidation/Private/DataValidationChangelist.cpp b/Engine/Plugins/Editor/DataValidation/Source/DataValidation/Private/DataValidationChangelist.cpp
--- a/Engine/Plugins/Editor/DataValidation/Source/DataValidation/Private/DataValidationChangelist.cpp
+++ b/Engine/Plugins/Editor/DataValidation/Source/DataValidation/Private/DataValidationChangelist.cpp
@@ -5,6 +5,7 @@
 #include "Algo/AnyOf.h"
 #include "Algo/Transform.h"
 #include "Misc/Paths.h"
+#include "Misc/DataValidation.h"
 #include "AssetRegistry/AssetData.h"
 #include "AssetRegistry/AssetRegistryModule.h"
 #include "ISourceControlProvider.h"
@@ -70,7 +71,7 @@ FString GetPrettyPackageName(const FName& InPackageName)
 	}
 }
 
-EDataValidationResult UDataValidationChangelist::IsDataValid(TArray<FText>& ValidationErrors, TArray<FText>& ValidationWarnings)
+EDataValidationResult UDataValidationChangelist::IsDataValid(FDataValidationContext& Context)
 {
 	ISourceControlProvider& SourceControlProvider = ISourceControlModule::Get().GetProvider();
 	
@@ -138,13 +139,13 @@ EDataValidationResult UDataValidationChangelist::IsDataValid(TArray<FText>& Vali
 		{
 			bHasChangelistErrors = true;
 			FText CurrentError = FText::Format(LOCTEXT("DataValidation.Changelist.Error", "{0} is missing from this changelist."), FText::FromString(GetPrettyPackageName(ExternalDependency)));
-			ValidationErrors.Add(CurrentError);
+			Context.AddError(CurrentError);
 		}
 		// Dependency is not at the latest revision
 		else if (!ExternalDependencyFileState->IsCurrent())
 		{
 			FText CurrentWarning = FText::Format(LOCTEXT("DataValidation.Changelist.NotLatest", "{0} is referenced but is not at the latest revision '{1}'"), FText::FromString(GetPrettyPackageName(ExternalDependency)), FText::FromString(ExternalPackageFilename));
-			ValidationWarnings.Add(CurrentWarning);
+			Context.AddWarning(CurrentWarning);
 		}
 		// Dependency is not in source control
 		else if (ExternalDependencyFileState->CanAdd())
@@ -153,13 +154,13 @@ EDataValidationResult UDataValidationChangelist::IsDataValid(TArray<FText>& Vali
 			{
 				bHasChangelistErrors = true;
 				FText CurrentError = FText::Format(LOCTEXT("DataValidation.Changelist.NotInWorkspace", "{0} is referenced and cannot be found in workspace '{1}'"), FText::FromString(GetPrettyPackageName(ExternalDependency)), FText::FromString(ExternalPackageFilename));
-				ValidationErrors.Add(CurrentError);
+				Context.AddError(CurrentError);
 			}
 			else
 			{
 				bHasChangelistErrors = true;
 				FText CurrentError = FText::Format(LOCTEXT("DataValidation.Changelist.NotInDepot", "{0} is referenced and must also be added to source control '{1}'"), FText::FromString(GetPrettyPackageName(ExternalDependency)), FText::FromString(ExternalPackageFilename));
-				ValidationErrors.Add(CurrentError);
+				Context.AddError(CurrentError);
 			}
 		}
 	}
diff --git a/Engine/Plugins/Editor/DataValidation/Source/DataValidation/Private/EditorValidatorSubsystem.cpp b/Engine/Plugins/Editor/DataValidation/Source/DataValidation/Private/EditorValidatorSubsystem.cpp
--- a/Engine/Plugins/Editor/DataValidation/Source/DataValidation/Private/EditorValidatorSubsystem.cpp
+++ b/Engine/Plugins/Editor/DataValidation/Source/DataValidation/Private/EditorValidatorSubsystem.cpp
@@ -15,6 +15,7 @@
 #include "CookOnTheSide/CookOnTheFlyServer.h"
 #include "Logging/MessageLog.h"
 #include "Misc/ScopedSlowTask.h"
+#include "Misc/DataValidation.h"
 #include "AssetRegistry/AssetData.h"
 #include "ISourceControlModule.h"
 #include "DataValidationChangelist.h"
@@ -158,7 +159,10 @@ EDataValidationResult UEditorValidatorSubsystem::IsObjectValid(UObject* InObject
 	if (ensure(InObject))
 	{
 		// First check the class level validation
-		Result = InObject->IsDataValid(ValidationErrors, ValidationWarnings);
+		FDataValidationContext Context;
+		Result = InObject->IsDataValid(Context);
+		Context.SplitIssues(ValidationWarnings, ValidationErrors);
+
 		// If the asset is still valid or there wasn't a class-level validation, keep validating with custom validators
 		if (Result != EDataValidationResult::Invalid)
 		{
diff --git a/Engine/Plugins/Editor/DataValidation/Source/DataValidation/Public/DataValidationChangelist.h b/Engine/Plugins/Editor/DataValidation/Source/DataValidation/Public/DataValidationChangelist.h
--- a/Engine/Plugins/Editor/DataValidation/Source/DataValidation/Public/DataValidationChangelist.h
+++ b/Engine/Plugins/Editor/DataValidation/Source/DataValidation/Public/DataValidationChangelist.h
@@ -26,7 +26,7 @@ public:
 	/** Default constructor with nothing */
 	UDataValidationChangelist() = default;
 
-	virtual EDataValidationResult IsDataValid(TArray<FText>& ValidationErrors, TArray<FText>& ValidationWarnings) override;
+	virtual EDataValidationResult IsDataValid(class FDataValidationContext& Context) override;
 
 	/** Initializes internal state so the validation can be done */
 	void Initialize(FSourceControlChangelistPtr InChangelist);
diff --git a/Engine/Plugins/Experimental/Mutable/Source/CustomizableObject/Classes/MuCO/CustomizableObject.h b/Engine/Plugins/Experimental/Mutable/Source/CustomizableObject/Classes/MuCO/CustomizableObject.h
--- a/Engine/Plugins/Experimental/Mutable/Source/CustomizableObject/Classes/MuCO/CustomizableObject.h
+++ b/Engine/Plugins/Experimental/Mutable/Source/CustomizableObject/Classes/MuCO/CustomizableObject.h
@@ -1334,9 +1334,7 @@ public:
 #if WITH_EDITOR
 
 	// UObject Interface -> Data validation
-	virtual EDataValidationResult IsDataValid(TArray<FText>& ValidationErrors) override;
-	virtual EDataValidationResult
-	IsDataValid(TArray<FText>& ValidationErrors, TArray<FText>& ValidationWarnings) override;
+	virtual EDataValidationResult IsDataValid(class FDataValidationContext& Context) override;
 	// End of UObject Interface
 
 	// UObject Interface -> Asset saving
diff --git a/Engine/Plugins/Experimental/Mutable/Source/CustomizableObject/Private/MuCO/CustomizableObject.cpp b/Engine/Plugins/Experimental/Mutable/Source/CustomizableObject/Private/MuCO/CustomizableObject.cpp
--- a/Engine/Plugins/Experimental/Mutable/Source/CustomizableObject/Private/MuCO/CustomizableObject.cpp
+++ b/Engine/Plugins/Experimental/Mutable/Source/CustomizableObject/Private/MuCO/CustomizableObject.cpp
@@ -24,6 +24,7 @@
 #include "Misc/CoreMisc.h"
 #include "Misc/PackageName.h"
 #include "Misc/Paths.h"
+#include "Misc/DataValidation.h"
 #include "Modules/ModuleManager.h"
 #include "MuCO/CustomizableObjectInstance.h"
 #include "MuCO/CustomizableObjectPrivate.h"
@@ -1467,14 +1468,7 @@ void UCustomizableObject::PreSaveRoot(FObjectPreSaveRootContext ObjectSaveContex
 	bIsValidationTriggeredBySave = true;
 }
 
-EDataValidationResult UCustomizableObject::IsDataValid(TArray<FText>& ValidationErrors)
-{
-	TArray<FText> ValidationWarnings;
-	return IsDataValid(ValidationErrors,ValidationWarnings);
-}
-
-EDataValidationResult UCustomizableObject::IsDataValid(TArray<FText>& ValidationErrors,
-	TArray<FText>& ValidationWarnings)
+EDataValidationResult UCustomizableObject::IsDataValid(FDataValidationContext& Context)
 {
 	// This method seems to be designed to check data errors (like variables with unexpected values).
 	// Currently it does not check if the root that we are compiling has already been compiled during another validation.
@@ -1510,8 +1504,20 @@ EDataValidationResult UCustomizableObject::IsDataValid(TArray<FText>& Validation
 		// --------------------------------------------------------------------------
 
 		// Request the compilation warnings and errors and provide them to the caller
-		Compiler->GetCompilationMessages(ValidationWarnings,ValidationErrors);
-		
+		TArray<FText> ValidationErrors;
+		TArray<FText> ValidationWarnings;
+		Compiler->GetCompilationMessages(ValidationWarnings, ValidationErrors);
+
+		for (const FText& ValidationError : ValidationErrors)
+		{
+			Context.AddError(ValidationError);
+		}
+
+		for (const FText& ValidationWarning : ValidationWarnings)
+		{
+			Context.AddWarning(ValidationWarning);
+		}
+
 		// Check if the compilation was able to be performed (not if there were errors but if a critical issue appeared
 		// before starting the compilation as we know.
 		switch (Compiler->GetCompilationState())
diff --git a/Engine/Source/Runtime/CoreUObject/Private/UObject/Obj.cpp b/Engine/Source/Runtime/CoreUObject/Private/UObject/Obj.cpp
--- a/Engine/Source/Runtime/CoreUObject/Private/UObject/Obj.cpp
+++ b/Engine/Source/Runtime/CoreUObject/Private/UObject/Obj.cpp
@@ -16,6 +16,7 @@
 #include "Misc/CoreDelegates.h"
 #include "Misc/App.h"
 #include "Misc/ITransactionObjectAnnotation.h"
+#include "Misc/DataValidation.h"
 #include "Modules/ModuleManager.h"
 #include "UObject/ObjectMacros.h"
 #include "UObject/UObjectGlobals.h"
@@ -5005,9 +5006,17 @@ EDataValidationResult UObject::IsDataValid(TArray<FText>& ValidationErrors)
 	return EDataValidationResult::NotValidated;
 }
 
-EDataValidationResult UObject::IsDataValid(TArray<FText>& ValidationErrors, TArray<FText>& ValidationWarnings)
+EDataValidationResult UObject::IsDataValid(FDataValidationContext& Context)
 {
-	return IsDataValid(ValidationErrors);
+	TArray<FText> ValidationErrors;
+	const EDataValidationResult Result = IsDataValid(ValidationErrors);
+
+	for (const FText& Text : ValidationErrors)
+	{
+		Context.AddError(Text);
+	}
+
+	return Result;
 }
 
 #endif // WITH_EDITOR
diff --git a/Engine/Source/Runtime/CoreUObject/Public/Misc/DataValidation.h b/Engine/Source/Runtime/CoreUObject/Public/Misc/DataValidation.h
new file mode 100644
--- /dev/null
+++ b/Engine/Source/Runtime/CoreUObject/Public/Misc/DataValidation.h
@@ -0,0 +1,48 @@
+// Copyright Epic Games, Inc. All Rights Reserved.
+
+#pragma once
+
+#include "CoreMinimal.h"
+#include "Logging/TokenizedMessage.h"
+
+class FDataValidationContext : public FNoncopyable
+{
+public:
+	struct FIssue
+	{
+		FText Message;
+		EMessageSeverity::Type Severity;
+
+		FIssue(const FText& InMessage, EMessageSeverity::Type InSeverity)
+			: Message(InMessage)
+			, Severity(InSeverity)
+		{}
+	};
+
+	void AddWarning(const FText& Text) { Issues.Emplace(Text, EMessageSeverity::Warning); NumWarnings++; }
+	void AddError(const FText& Text) { Issues.Emplace(Text, EMessageSeverity::Error); NumErrors++; }
+
+	const TArray<FIssue>& GetIssues() const { return Issues; }
+	uint32 GetNumWarnings() const { return NumWarnings; }
+	uint32 GetNumErrors() const { return NumErrors; }
+
+	void SplitIssues(TArray<FText>& Warnings, TArray<FText>& Errors) const
+	{
+		for (const FIssue& Issue : GetIssues())
+		{
+			if (Issue.Severity == EMessageSeverity::Error)
+			{
+				Errors.Add(Issue.Message);
+			}
+			else if (Issue.Severity == EMessageSeverity::Warning)
+			{
+				Warnings.Add(Issue.Message);
+			}
+		}
+	}
+
+private:
+	TArray<FIssue> Issues;
+	uint32 NumWarnings = 0;
+	uint32 NumErrors = 0;
+};
\ No newline at end of file
diff --git a/Engine/Source/Runtime/CoreUObject/Public/UObject/Object.h b/Engine/Source/Runtime/CoreUObject/Public/UObject/Object.h
--- a/Engine/Source/Runtime/CoreUObject/Public/UObject/Object.h
+++ b/Engine/Source/Runtime/CoreUObject/Public/UObject/Object.h
@@ -1002,21 +1002,21 @@ public:
 	/**
 	 * Generic function to validate objects during changelist validations, etc.
 	 *
-	 * @param	ValidationErrors	the array of validation errors.
+	 * @param	Context	the context holding validation warnings/errors.
 	 * @return Valid if this object has data validation rules set up for it and the data for this object is valid. Returns Invalid if it does not pass 
 	 *         the rules. Returns NotValidated if no rules are set for this object.
 	 */
-	virtual EDataValidationResult IsDataValid(TArray<FText>& ValidationErrors);
+	virtual EDataValidationResult IsDataValid(class FDataValidationContext& Context);
 
 	/**
 	 * Generic function to validate objects during changelist validations, etc.
 	 *
 	 * @param	ValidationErrors	the array of validation errors.
-	 * @param	ValidationWarnings	the array of validation warnings.
 	 * @return Valid if this object has data validation rules set up for it and the data for this object is valid. Returns Invalid if it does not pass 
 	 *         the rules. Returns NotValidated if no rules are set for this object.
+	 * @note	Will be deprecated in 5.2 in favor of version taking a FDataValidationContext.
 	 */
-	virtual EDataValidationResult IsDataValid(TArray<FText>& ValidationErrors, TArray<FText>& ValidationWarnings);
+	virtual EDataValidationResult IsDataValid(TArray<FText>& ValidationErrors);
 #endif // WITH_EDITOR
 
 	/*----------------------------------------------------------
