DataValidationCommandlet: Initialize the AssetRegistry since we're using it to get a list of all assets. Add -projectonly argument.
#rnx

diff --git a/Engine/Plugins/Editor/DataValidation/Source/DataValidation/Private/DataValidationCommandlet.cpp b/Engine/Plugins/Editor/DataValidation/Source/DataValidation/Private/DataValidationCommandlet.cpp
--- a/Engine/Plugins/Editor/DataValidation/Source/DataValidation/Private/DataValidationCommandlet.cpp
+++ b/Engine/Plugins/Editor/DataValidation/Source/DataValidation/Private/DataValidationCommandlet.cpp
@@ -1,12 +1,16 @@
 // Copyright Epic Games, Inc. All Rights Reserved.
 
 #include "DataValidationCommandlet.h"
-#include "DataValidationModule.h"
-#include "AssetRegistry/IAssetRegistry.h"
+
 #include "AssetRegistry/AssetRegistryHelpers.h"
 #include "AssetRegistry/AssetRegistryModule.h"
+#include "AssetRegistry/IAssetRegistry.h"
+#include "DataValidationModule.h"
 #include "Editor.h"
 #include "EditorValidatorSubsystem.h"
+#include "Misc/PackageName.h"
+#include "Misc/Paths.h"
+#include "Misc/PathViews.h"
 
 #include UE_INLINE_GENERATED_CPP_BY_NAME(DataValidationCommandlet)
 
@@ -40,6 +44,11 @@ bool UDataValidationCommandlet::ValidateData(const FString& FullCommandLine)
 	ParseCommandLine(*FullCommandLine, Tokens, Switches, Params);
 
 	FAssetRegistryModule& AssetRegistryModule = FModuleManager::LoadModuleChecked<FAssetRegistryModule>(AssetRegistryConstants::ModuleName);
+	IAssetRegistry& AssetRegistry = AssetRegistryModule.Get();
+	AssetRegistry.SearchAllAssets(true /* bSynchronousSearch */);
+
+	bool bProjectOnly = !Switches.Contains(TEXT("includeengine"));
+
 	TArray<FAssetData> AssetDataList;
 	FString AssetTypeString;
 	if (FParse::Value(*FullCommandLine, TEXT("AssetType="), AssetTypeString) && !AssetTypeString.IsEmpty())
@@ -58,11 +67,33 @@ bool UDataValidationCommandlet::ValidateData(const FString& FullCommandLine)
 			}
 		}
 		const bool bSearchSubClasses = true;
-		AssetRegistryModule.Get().GetAssetsByClass(FTopLevelAssetPath(AssetTypeString), AssetDataList, bSearchSubClasses);
+		AssetRegistry.GetAssetsByClass(FTopLevelAssetPath(AssetTypeString), AssetDataList, bSearchSubClasses);
 	}
 	else
 	{
-		AssetRegistryModule.Get().GetAllAssets(AssetDataList);
+		AssetRegistry.GetAllAssets(AssetDataList);
+	}
+
+	if (bProjectOnly)
+	{
+		FString EngineDir = FPaths::ConvertRelativePathToFull(FPaths::EngineDir());
+		AssetDataList.RemoveAll([&EngineDir](const FAssetData& AssetData)
+			{
+				// Remove /Engine and any plugins from /Engine, but keep /Game and any plugins under /Game.
+				FString FileName;
+				FString PackageName;
+				AssetData.PackageName.ToString(PackageName);
+				if (!FPackageName::TryConvertLongPackageNameToFilename(PackageName, FileName))
+				{
+					// We don't recognize this packagepath, so keep it
+					return false;
+				}
+				// ConvertLongPackageNameToFilename can return ../../Plugins for some plugins instead of
+				// ../../../Engine/Plugins. We should fix that in FPackageName to always return the normalized
+				// filename. For now, workaround it by converting to absolute paths.
+				FileName = FPaths::ConvertRelativePathToFull(MoveTemp(FileName));
+				return FPathViews::IsParentPathOf(EngineDir, FileName);
+			});
 	}
 
 	UEditorValidatorSubsystem* EditorValidationSubsystem = GEditor->GetEditorSubsystem<UEditorValidatorSubsystem>();
