Added option to disable the 'out of date dependency' warning

diff --git a/Engine/Plugins/Editor/DataValidation/Source/DataValidation/Private/DataValidationChangelist.cpp b/Engine/Plugins/Editor/DataValidation/Source/DataValidation/Private/DataValidationChangelist.cpp
--- a/Engine/Plugins/Editor/DataValidation/Source/DataValidation/Private/DataValidationChangelist.cpp
+++ b/Engine/Plugins/Editor/DataValidation/Source/DataValidation/Private/DataValidationChangelist.cpp
@@ -13,6 +13,7 @@
 #include "SourceControlHelpers.h"
 #include "SourceControlOperations.h"
 #include "UncontrolledChangelistsModule.h"
+#include "Misc/ConfigCacheIni.h"
 
 #include UE_INLINE_GENERATED_CPP_BY_NAME(DataValidationChangelist)
 
@@ -121,6 +122,9 @@ EDataValidationResult UDataValidationChangelist::IsDataValid(FDataValidationCont
 
 	check(ExternalDependenciesFilenames.Num() == ExternalDependencies.Num());
 
+	bool bIgnoreOutOfDateDependencies = false;
+	GConfig->GetBool(TEXT("DataValidationChangelistSettings"), TEXT("bIgnoreOutOfDateDependencies"), bIgnoreOutOfDateDependencies, GEditorIni);
+
 	for (int32 i = 0; i < ExternalDependenciesFilenames.Num(); ++i)
 	{
 		const FString& ExternalPackageFilename = ExternalDependenciesFilenames[i];
@@ -144,8 +148,11 @@ EDataValidationResult UDataValidationChangelist::IsDataValid(FDataValidationCont
 		// Dependency is not at the latest revision
 		else if (!ExternalDependencyFileState->IsCurrent())
 		{
-			FText CurrentWarning = FText::Format(LOCTEXT("DataValidation.Changelist.NotLatest", "{0} is referenced but is not at the latest revision '{1}'"), FText::FromString(GetPrettyPackageName(ExternalDependency)), FText::FromString(ExternalPackageFilename));
-			Context.AddWarning(CurrentWarning);
+			if (!bIgnoreOutOfDateDependencies)
+			{
+				FText CurrentWarning = FText::Format(LOCTEXT("DataValidation.Changelist.NotLatest", "{0} is referenced but is not at the latest revision '{1}'"), FText::FromString(GetPrettyPackageName(ExternalDependency)), FText::FromString(ExternalPackageFilename));
+				Context.AddWarning(CurrentWarning);
+			}
 		}
 		// Dependency is not in source control
 		else if (ExternalDependencyFileState->CanAdd())
