Convert data validation content browser context menus to ToolMenus
#rnx

diff --git a/Engine/Plugins/Editor/DataValidation/Source/DataValidation/Private/DataValidationModule.cpp b/Engine/Plugins/Editor/DataValidation/Source/DataValidation/Private/DataValidationModule.cpp
--- a/Engine/Plugins/Editor/DataValidation/Source/DataValidation/Private/DataValidationModule.cpp
+++ b/Engine/Plugins/Editor/DataValidation/Source/DataValidation/Private/DataValidationModule.cpp
@@ -18,6 +18,7 @@
 #include "Misc/OutputDeviceConsole.h"
 #include "Modules/ModuleManager.h"
 
+#include "ContentBrowserMenuContexts.h"
 #include "ContentBrowserModule.h"
 #include "ContentBrowserDelegates.h"
 
@@ -39,13 +40,9 @@ class FDataValidationModule : public IDataValidationModule
 	/** Validates selected assets and opens a window to report the results. If bValidateDependencies is true it will also validate any assets that the selected assets depend on. */
 	virtual void ValidateAssets(const TArray<FAssetData>& SelectedAssets, bool bValidateDependencies, const EDataValidationUsecase InValidationUsecase) override;
 
-	void ValidateFolders(TArray<FString> SelectedFolders, const EDataValidationUsecase InValidationUsecase);
+	void ValidateFolders(const TArray<FString>& SelectedFolders, const EDataValidationUsecase InValidationUsecase);
 
 private:
-	TSharedRef<FExtender> OnExtendContentBrowserAssetSelectionMenu(const TArray<FAssetData>& SelectedAssets);
-	TSharedRef<FExtender> OnExtendContentBrowserPathSelectionMenu(const TArray<FString>& SelectedPaths);
-	void CreateDataValidationContentBrowserAssetMenu(FMenuBuilder& MenuBuilder, TArray<FAssetData> SelectedAssets);
-	void CreateDataValidationContentBrowserPathMenu(FMenuBuilder& MenuBuilder, TArray<FString> SelectedPaths);
 	void OnPackageSaved(const FString& PackageFileName, UPackage* Package, FObjectPostSaveContext ObjectSaveContext);
 
 	// Adds Asset and any assets it depends on to the set DependentAssets
@@ -54,9 +51,6 @@ private:
 	void RegisterMenus();
 	static FText Menu_ValidateDataGetTitle();
 	static void Menu_ValidateData();
-
-	FDelegateHandle ContentBrowserAssetExtenderDelegateHandle;
-	FDelegateHandle ContentBrowserPathExtenderDelegateHandle;
 };
 
 IMPLEMENT_MODULE(FDataValidationModule, DataValidation)
@@ -65,20 +59,8 @@ void FDataValidationModule::StartupModule()
 {	
 	if (!IsRunningCommandlet() && !IsRunningGame() && FSlateApplication::IsInitialized())
 	{
-		// Register content browser hook
-		FContentBrowserModule& ContentBrowserModule = FModuleManager::LoadModuleChecked<FContentBrowserModule>(TEXT("ContentBrowser"));
-		TArray<FContentBrowserMenuExtender_SelectedAssets>& CBAssetMenuExtenderDelegates = ContentBrowserModule.GetAllAssetViewContextMenuExtenders();
-
-		CBAssetMenuExtenderDelegates.Add(FContentBrowserMenuExtender_SelectedAssets::CreateRaw(this, &FDataValidationModule::OnExtendContentBrowserAssetSelectionMenu));
-		ContentBrowserAssetExtenderDelegateHandle = CBAssetMenuExtenderDelegates.Last().GetHandle();
-
-		TArray<FContentBrowserMenuExtender_SelectedPaths>& CBFolderMenuExtenderDelegates = ContentBrowserModule.GetAllPathViewContextMenuExtenders();
-
-		CBFolderMenuExtenderDelegates.Add(FContentBrowserMenuExtender_SelectedPaths::CreateRaw(this, &FDataValidationModule::OnExtendContentBrowserPathSelectionMenu));
-		ContentBrowserPathExtenderDelegateHandle = CBFolderMenuExtenderDelegates.Last().GetHandle();
-
 		// add the File->DataValidation menu subsection
-		FCoreDelegates::OnPostEngineInit.AddRaw(this, &FDataValidationModule::RegisterMenus);
+		UToolMenus::Get()->RegisterStartupCallback(FSimpleMulticastDelegate::FDelegate::CreateRaw(this, &FDataValidationModule::RegisterMenus));
 
 		// Add save callback
 		UPackage::PackageSavedWithContextEvent.AddRaw(this, &FDataValidationModule::OnPackageSaved);
@@ -96,17 +78,8 @@ void FDataValidationModule::ShutdownModule()
 {
 	if (!IsRunningCommandlet() && !IsRunningGame() && !IsRunningDedicatedServer())
 	{
-		FContentBrowserModule* ContentBrowserModule = FModuleManager::GetModulePtr<FContentBrowserModule>(TEXT("ContentBrowser"));
-		if (ContentBrowserModule)
-		{
-			TArray<FContentBrowserMenuExtender_SelectedAssets>& CBMenuExtenderDelegates = ContentBrowserModule->GetAllAssetViewContextMenuExtenders();
-			CBMenuExtenderDelegates.RemoveAll([this](const FContentBrowserMenuExtender_SelectedAssets& Delegate) { return Delegate.GetHandle() == ContentBrowserAssetExtenderDelegateHandle; });
-			CBMenuExtenderDelegates.RemoveAll([this](const FContentBrowserMenuExtender_SelectedAssets& Delegate) { return Delegate.GetHandle() == ContentBrowserPathExtenderDelegateHandle; });
-		}
-
 		// remove menu extension
 		UToolMenus::UnregisterOwner(this);
-		FCoreDelegates::OnPostEngineInit.RemoveAll(this);
 
 		UPackage::PackageSavedWithContextEvent.RemoveAll(this);
 	}
@@ -115,15 +88,85 @@ void FDataValidationModule::ShutdownModule()
 void FDataValidationModule::RegisterMenus()
 {
 	FToolMenuOwnerScoped OwnerScoped(this);
-	UToolMenu* Menu = UToolMenus::Get()->ExtendMenu("LevelEditor.MainMenu.Tools");
-	FToolMenuSection& Section = Menu->AddSection("DataValidation", LOCTEXT("DataValidation", "DataValidation"));
-	Section.AddEntry(FToolMenuEntry::InitMenuEntry(
-		"ValidateData",
-		TAttribute<FText>::Create(&Menu_ValidateDataGetTitle),
-		LOCTEXT("ValidateDataTooltip", "Validates all user data in content directory."),
-		FSlateIcon(FAppStyle::GetAppStyleSetName(), "DeveloperTools.MenuIcon"),
-		FUIAction(FExecuteAction::CreateStatic(&FDataValidationModule::Menu_ValidateData))
-	));
+
+	{
+		UToolMenu* Menu = UToolMenus::Get()->ExtendMenu("LevelEditor.MainMenu.Tools");
+		FToolMenuSection& Section = Menu->AddSection("DataValidation", LOCTEXT("DataValidation", "DataValidation"));
+		Section.AddEntry(FToolMenuEntry::InitMenuEntry(
+			"ValidateData",
+			TAttribute<FText>::Create(&Menu_ValidateDataGetTitle),
+			LOCTEXT("ValidateDataTooltip", "Validates all user data in content directory."),
+			FSlateIcon(FAppStyle::GetAppStyleSetName(), "DeveloperTools.MenuIcon"),
+			FUIAction(FExecuteAction::CreateStatic(&FDataValidationModule::Menu_ValidateData))
+		));
+	}
+
+	{
+		UToolMenu* Menu = UToolMenus::Get()->ExtendMenu("ContentBrowser.AssetContextMenu.AssetActionsSubMenu");
+		FToolMenuSection& Section = Menu->FindOrAddSection("AssetContextAdvancedActions");
+		Section.AddMenuEntry(
+			"ValidateAssets",
+			LOCTEXT("ValidateAssetsTabTitle", "Validate Assets"),
+			LOCTEXT("ValidateAssetsTooltipText", "Runs data validation on these assets."),
+			FSlateIcon(),
+			FToolMenuExecuteAction::CreateLambda([this](const FToolMenuContext& InContext)
+			{
+				if (UContentBrowserAssetContextMenuContext* Context = InContext.FindContext<UContentBrowserAssetContextMenuContext>())
+				{
+					ValidateAssets(Context->SelectedAssets, false, EDataValidationUsecase::Manual);
+				}
+			})
+		);
+
+		Section.AddMenuEntry(
+			"ValidateAssetsAndDependencies",
+			LOCTEXT("ValidateAssetsAndDependenciesTabTitle", "Validate Assets and Dependencies"),
+			LOCTEXT("ValidateAssetsAndDependenciesTooltipText", "Runs data validation on these assets and all assets they depend on."),
+			FSlateIcon(),
+			FToolMenuExecuteAction::CreateLambda([this](const FToolMenuContext& InContext)
+			{
+				if (UContentBrowserAssetContextMenuContext* Context = InContext.FindContext<UContentBrowserAssetContextMenuContext>())
+				{
+					ValidateAssets(Context->SelectedAssets, true, EDataValidationUsecase::Manual);
+				}
+			})
+		);
+	}
+
+	{
+		UToolMenu* Menu = UToolMenus::Get()->ExtendMenu("ContentBrowser.FolderContextMenu");
+		FToolMenuSection& Section = Menu->FindOrAddSection("PathContextBulkOperations");
+		Section.AddMenuEntry(
+			"ValidateAssetsPath",
+			LOCTEXT("ValidateAssetsPathTabTitle", "Validate Assets in Folder"),
+			LOCTEXT("ValidateAssetsPathTooltipText", "Runs data validation on the assets in the selected folder."),
+			FSlateIcon(FAppStyle::GetAppStyleSetName(), "Icons.Validate"),
+			FToolMenuExecuteAction::CreateLambda([this](const FToolMenuContext& InContext)
+			{
+				if (UContentBrowserFolderContext* Context = InContext.FindContext<UContentBrowserFolderContext>())
+				{
+					const TArray<FString>& SelectedPaths = Context->GetSelectedPackagePaths();
+
+					FString FormattedSelectedPaths;
+					for (int32 i = 0; i < SelectedPaths.Num(); ++i)
+					{
+						FormattedSelectedPaths.Append(SelectedPaths[i]);
+						if (i < SelectedPaths.Num() - 1)
+						{
+							FormattedSelectedPaths.Append(LINE_TERMINATOR);
+						}
+					}
+					FFormatNamedArguments Args;
+					Args.Add(TEXT("Paths"), FText::FromString(FormattedSelectedPaths));
+					const EAppReturnType::Type Result = FMessageDialog::Open(EAppMsgType::YesNo, FText::Format(LOCTEXT("DataValidationConfirmation", "Are you sure you want to proceed with validating the following folders?\n\n{Paths}"), Args));
+					if (Result == EAppReturnType::Yes)
+					{
+						ValidateFolders(SelectedPaths, EDataValidationUsecase::Manual);
+					}
+				}
+			})
+		);
+	}
 }
 
 FText FDataValidationModule::Menu_ValidateDataGetTitle()
@@ -218,7 +261,7 @@ void FDataValidationModule::ValidateAssets(const TArray<FAssetData>& SelectedAss
 	}
 }
 
-void FDataValidationModule::ValidateFolders(TArray<FString> SelectedFolders, const EDataValidationUsecase InValidationUsecase)
+void FDataValidationModule::ValidateFolders(const TArray<FString>& SelectedFolders, const EDataValidationUsecase InValidationUsecase)
 {
 	FAssetRegistryModule& AssetRegistryModule = FModuleManager::Get().LoadModuleChecked<FAssetRegistryModule>(TEXT("AssetRegistry"));
 
@@ -250,84 +293,6 @@ void FDataValidationModule::ValidateFolders(TArray<FString> SelectedFolders, con
 	ValidateAssets(AssetList, false, InValidationUsecase);
 }
 
-// Extend content browser menu for groups of selected assets
-TSharedRef<FExtender> FDataValidationModule::OnExtendContentBrowserAssetSelectionMenu(const TArray<FAssetData>& SelectedAssets)
-{
-	TSharedRef<FExtender> Extender(new FExtender());
-
-	Extender->AddMenuExtension(
-		"AssetContextAdvancedActions",
-		EExtensionHook::After,
-		nullptr,
-		FMenuExtensionDelegate::CreateRaw(this, &FDataValidationModule::CreateDataValidationContentBrowserAssetMenu, SelectedAssets));
-
-	return Extender;
-}
-
-// Extend content browser menu for groups of selected assets
-void FDataValidationModule::CreateDataValidationContentBrowserAssetMenu(FMenuBuilder& MenuBuilder, TArray<FAssetData> SelectedAssets)
-{
-	MenuBuilder.AddMenuEntry
-	(
-		LOCTEXT("ValidateAssetsTabTitle", "Validate Assets"),
-		LOCTEXT("ValidateAssetsTooltipText", "Runs data validation on these assets."),
-		FSlateIcon(),
-		FUIAction(FExecuteAction::CreateLambda([this, SelectedAssets]() { ValidateAssets(SelectedAssets, false, EDataValidationUsecase::Manual); }))
-	);
-
-	MenuBuilder.AddMenuEntry
-	(
-		LOCTEXT("ValidateAssetsAndDependenciesTabTitle", "Validate Assets and Dependencies"),
-		LOCTEXT("ValidateAssetsAndDependenciesTooltipText", "Runs data validation on these assets and all assets they depend on."),
-		FSlateIcon(),
-		FUIAction(FExecuteAction::CreateLambda([this, SelectedAssets]() { ValidateAssets(SelectedAssets, true, EDataValidationUsecase::Manual); }))
-	);
-}
-
-// Extend content browser menu for groups of asset paths (folders)
-TSharedRef<FExtender> FDataValidationModule::OnExtendContentBrowserPathSelectionMenu(const TArray<FString>& SelectedPaths)
-{
-	TSharedRef<FExtender> Extender(new FExtender());
-
-	Extender->AddMenuExtension(
-		"PathContextBulkOperations",
-		EExtensionHook::After,
-		nullptr,
-		FMenuExtensionDelegate::CreateRaw(this, &FDataValidationModule::CreateDataValidationContentBrowserPathMenu, SelectedPaths));
-
-	return Extender;
-}
-
-// Extend content browser menu for groups of asset paths (folders)
-void FDataValidationModule::CreateDataValidationContentBrowserPathMenu(FMenuBuilder& MenuBuilder, TArray<FString> SelectedPaths)
-{
-	MenuBuilder.AddMenuEntry
-	(
-		LOCTEXT("ValidateAssetsPathTabTitle", "Validate Assets in Folder"),
-		LOCTEXT("ValidateAssetsPathTooltipText", "Runs data validation on the assets in the selected folder."),
-		FSlateIcon(FAppStyle::GetAppStyleSetName(), "Icons.Validate"),
-		FUIAction(FExecuteAction::CreateLambda([this, SelectedPaths]
-		{
-			FString FormattedSelectedPaths;
-			for (int32 i = 0; i < SelectedPaths.Num(); ++i)
-			{
-				FormattedSelectedPaths.Append(SelectedPaths[i]);
-				if (i < SelectedPaths.Num() - 1)
-				{
-					FormattedSelectedPaths.Append(LINE_TERMINATOR);
-				}
-			}
-			FFormatNamedArguments Args;
-			Args.Add(TEXT("Paths"), FText::FromString(FormattedSelectedPaths));
-			const EAppReturnType::Type Result = FMessageDialog::Open(EAppMsgType::YesNo, FText::Format(LOCTEXT("DataValidationConfirmation", "Are you sure you want to proceed with validating the following folders?\n\n{Paths}"), Args));
-			if (Result == EAppReturnType::Yes)
-			{
-				ValidateFolders(SelectedPaths, EDataValidationUsecase::Manual);
-			}
-		}))
-	);
-}
-
 void FDataValidationModule::OnPackageSaved(const FString& PackageFileName, UPackage* Package, FObjectPostSaveContext ObjectSaveContext)
 {
 	UEditorValidatorSubsystem* EditorValidationSubsystem = GEditor->GetEditorSubsystem<UEditorValidatorSubsystem>();
