Changelist validation now perform streaming generation validation on content bundle as well

diff --git a/Engine/Plugins/Editor/DataValidation/Source/DataValidation/Private/WorldPartitionChangelistValidator.cpp b/Engine/Plugins/Editor/DataValidation/Source/DataValidation/Private/WorldPartitionChangelistValidator.cpp
--- a/Engine/Plugins/Editor/DataValidation/Source/DataValidation/Private/WorldPartitionChangelistValidator.cpp
+++ b/Engine/Plugins/Editor/DataValidation/Source/DataValidation/Private/WorldPartitionChangelistValidator.cpp
@@ -10,12 +10,14 @@
 #include "SourceControlHelpers.h"
 
 #include "Engine/World.h"
+#include "Misc/PathViews.h"
 #include "WorldPartition/WorldPartitionActorDesc.h"
 #include "WorldPartition/WorldPartitionActorDescUtils.h"
 #include "WorldPartition/ActorDescContainer.h"
 #include "WorldPartition/WorldPartition.h"
 #include "WorldPartition/DataLayer/DataLayerInstanceWithAsset.h"
 #include "WorldPartition/DataLayer/WorldDataLayers.h"
+#include "WorldPartition/ContentBundle/ContentBundlePaths.h"
 
 #include UE_INLINE_GENERATED_CPP_BY_NAME(WorldPartitionChangelistValidator)
 
@@ -164,29 +166,58 @@ EDataValidationResult UWorldPartitionChangelistValidator::ValidateActorsAndDataL
 		}
 	}
 
-	// For Each world 
-	for (TTuple<FTopLevelAssetPath, TSet<FAssetData>>& It : MapToActorsFiles)
+	auto RegisterContainerToValidate = [](UWorld* InWorld, FName InContainerPackageName, FActorDescContainerCollection& OutRegisteredContainers)
 	{
-		const FTopLevelAssetPath& MapPath = It.Get<0>();
-		const TSet<FAssetData>& ActorsData = It.Get<1>();
-
-		// Find/Load the ActorDescContainer
-		UWorld* World = FindObject<UWorld>(nullptr, *MapPath.ToString(), true);
+		if (OutRegisteredContainers.Contains(InContainerPackageName))
+		{
+			return;
+		}
 
 		UActorDescContainer* ActorDescContainer = nullptr;
-		
-		if (World != nullptr)
+		if (InWorld != nullptr)
 		{
-			// World is Loaded reuse the ActorDescContainer of the World
-			ActorDescContainer = World->GetWorldPartition()->GetActorDescContainer();
+			// World is Loaded reuse the ActorDescContainer of the Content Bundle
+			ActorDescContainer = InWorld->GetWorldPartition()->FindContainer(InContainerPackageName);
 		}
-		
+
 		// Even if world is valid, its world partition is not necessarily initialized
 		if (!ActorDescContainer)
 		{
 			// Find in memory failed, load the ActorDescContainer
 			ActorDescContainer = NewObject<UActorDescContainer>();
-			ActorDescContainer->Initialize({ nullptr, MapPath.GetPackageName() });
+			ActorDescContainer->Initialize({ nullptr, InContainerPackageName });
+		}
+
+		OutRegisteredContainers.AddContainer(ActorDescContainer);
+	};
+
+	// For Each world 
+	for (TTuple<FTopLevelAssetPath, TSet<FAssetData>>& It : MapToActorsFiles)
+	{
+		const FTopLevelAssetPath& MapPath = It.Get<0>();
+		const TSet<FAssetData>& ActorsData = It.Get<1>();
+
+		// Find/Load the ActorDescContainer
+		UWorld* World = FindObject<UWorld>(nullptr, *MapPath.ToString(), true);
+		
+		FActorDescContainerCollection ContainersToValidate;
+		for (const FAssetData& ActorData : ActorsData)
+		{
+			FString ActorPackagePath = ActorData.PackagePath.ToString();
+			if (ContentBundlePaths::IsAContentBundlePackagePath(ActorPackagePath))
+			{
+				FStringView ContentBundleMountPoint = FPathViews::GetMountPointNameFromPath(ActorPackagePath);
+				FGuid ContentBundleGuid = ContentBundlePaths::GetContentBundleGuidFromExternalActorPackagePath(ActorPackagePath);
+				
+				FString ContentBundleContainerPackagePath;
+				verify(ContentBundlePaths::BuildActorDescContainerPackgePath(FString(ContentBundleMountPoint), ContentBundleGuid, MapPath.GetPackageName().ToString(), ContentBundleContainerPackagePath));
+
+				RegisterContainerToValidate(World, FName(*ContentBundleContainerPackagePath), ContainersToValidate);
+			}
+			else
+			{
+				RegisterContainerToValidate(World, MapPath.GetPackageName(), ContainersToValidate);
+			}
 		}
 
 		// Build a set of Relevant Actor Guids to scope error messages to what's contained in the CL 
@@ -196,7 +227,7 @@ EDataValidationResult UWorldPartitionChangelistValidator::ValidateActorsAndDataL
 		for (const FAssetData& ActorData : ActorsData)
 		{
 			// Get the FWorldPartitionActor			
-			const FWorldPartitionActorDesc* ActorDesc = ActorDescContainer->GetActorDesc(ActorData.AssetName.ToString());
+			const FWorldPartitionActorDesc* ActorDesc = ContainersToValidate.GetActorDesc(ActorData.AssetName.ToString());
 
 			if (ActorDesc != nullptr)
 			{
@@ -207,7 +238,11 @@ EDataValidationResult UWorldPartitionChangelistValidator::ValidateActorsAndDataL
 		// Invoke static WorldPartition Validation from the ActorDescContainer
 		const bool bIsStreamingDisabled = ULevel::GetIsStreamingDisabledFromPackage(MapPath.GetPackageName());
 		const bool bIsChangelistValidation = true;
-		UWorldPartition::CheckForErrors(this, ActorDescContainer, !bIsStreamingDisabled, bIsChangelistValidation);
+
+		ContainersToValidate.ForEachActorDescContainer([this, bIsStreamingDisabled, bIsChangelistValidation](const UActorDescContainer* ActorDescContainer)
+		{
+			UWorldPartition::CheckForErrors(this, ActorDescContainer, !bIsStreamingDisabled, bIsChangelistValidation);
+		});
 	}
 
 	if (Errors->Num())
diff --git a/Engine/Source/Runtime/Engine/Private/WorldPartition/ActorDescContainerCollection.cpp b/Engine/Source/Runtime/Engine/Private/WorldPartition/ActorDescContainerCollection.cpp
--- a/Engine/Source/Runtime/Engine/Private/WorldPartition/ActorDescContainerCollection.cpp
+++ b/Engine/Source/Runtime/Engine/Private/WorldPartition/ActorDescContainerCollection.cpp
@@ -22,10 +22,10 @@ bool FActorDescContainerCollection::RemoveContainer(UActorDescContainer* Contain
 
 bool FActorDescContainerCollection::Contains(const FName& ContainerPackageName) const
 {
-	 return Find(ContainerPackageName) != nullptr;
+	 return FindContainer(ContainerPackageName) != nullptr;
 }
 
-UActorDescContainer* FActorDescContainerCollection::Find(const FName& ContainerPackageName) const
+UActorDescContainer* FActorDescContainerCollection::FindContainer(const FName& ContainerPackageName) const
 {
 	auto* ContainerPtr = ActorDescContainerCollection.FindByPredicate([&ContainerPackageName](const UActorDescContainer* ActorDescContainer) { return ActorDescContainer->GetContainerPackage() == ContainerPackageName; });
 	return ContainerPtr != nullptr ? *ContainerPtr : nullptr;
diff --git a/Engine/Source/Runtime/Engine/Public/WorldPartition/ActorDescContainerCollection.h b/Engine/Source/Runtime/Engine/Public/WorldPartition/ActorDescContainerCollection.h
--- a/Engine/Source/Runtime/Engine/Public/WorldPartition/ActorDescContainerCollection.h
+++ b/Engine/Source/Runtime/Engine/Public/WorldPartition/ActorDescContainerCollection.h
@@ -15,7 +15,7 @@ public:
 	void AddContainer(UActorDescContainer* Container);
 	bool RemoveContainer(UActorDescContainer* Container);
 	bool Contains(const FName& ContainerPackageName) const;
-	UActorDescContainer* Find(const FName& ContainerPackageName) const;
+	UActorDescContainer* FindContainer(const FName& ContainerPackageName) const;
 
 	void Empty() { ActorDescContainerCollection.Empty(); }
 	bool IsEmpty() const { return ActorDescContainerCollection.IsEmpty(); }
