Skip editor validation on cooked assets

diff --git a/Engine/Plugins/Editor/DataValidation/Source/DataValidation/Private/EditorValidatorSubsystem.cpp b/Engine/Plugins/Editor/DataValidation/Source/DataValidation/Private/EditorValidatorSubsystem.cpp
--- a/Engine/Plugins/Editor/DataValidation/Source/DataValidation/Private/EditorValidatorSubsystem.cpp
+++ b/Engine/Plugins/Editor/DataValidation/Source/DataValidation/Private/EditorValidatorSubsystem.cpp
@@ -276,6 +276,12 @@ int32 UEditorValidatorSubsystem::ValidateAssetsWithSettings(const TArray<FAssetD
 		FText ValidatingMessage = FText::Format(LOCTEXT("ValidatingFilename", "Validating {0}"), FText::FromString(Data.GetFullName()));
 		SlowTask.EnterProgressFrame(1.0f / NumFilesToValidate, ValidatingMessage);
 
+		if (Data.HasAnyPackageFlags(PKG_Cooked))
+		{
+			++NumFilesSkipped;
+			continue;
+		}
+
 		// Check exclusion path
 		if (InSettings.bSkipExcludedDirectories && IsPathExcludedFromValidation(Data.PackageName.ToString()))
 		{
