Editor's "Check Content" validation fails when referencing a Verse class.
References to verse type packages should be simply skipped in the verifier.

diff --git a/Engine/Plugins/Editor/DataValidation/Source/DataValidation/Private/DataValidationChangelist.cpp b/Engine/Plugins/Editor/DataValidation/Source/DataValidation/Private/DataValidationChangelist.cpp
--- a/Engine/Plugins/Editor/DataValidation/Source/DataValidation/Private/DataValidationChangelist.cpp
+++ b/Engine/Plugins/Editor/DataValidation/Source/DataValidation/Private/DataValidationChangelist.cpp
@@ -135,6 +135,13 @@ EDataValidationResult UDataValidationChangelist::IsDataValid(FDataValidationCont
 			continue;
 		}
 
+		// Some kinds of package dependencies create false positives because they're not checked into source
+		// control as of this writing (1/23/2023). We just skip over them.
+		if (FPackageName::IsVersePackage(ExternalDependency.ToString()))
+		{
+			continue;
+		}
+
 		// Dependency is checked out or added but is not in this changelist
 		if (ExternalDependencyFileState->IsCheckedOut() || ExternalDependencyFileState->IsAdded())
 		{
