Added more missing generated.h includes.

diff --git a/Engine/Plugins/Editor/DataValidation/Source/DataValidation/Public/DataValidationModule.h b/Engine/Plugins/Editor/DataValidation/Source/DataValidation/Public/DataValidationModule.h
--- a/Engine/Plugins/Editor/DataValidation/Source/DataValidation/Public/DataValidationModule.h
+++ b/Engine/Plugins/Editor/DataValidation/Source/DataValidation/Public/DataValidationModule.h
@@ -5,6 +5,8 @@
 #include "Modules/ModuleManager.h"
 #include "UObject/ObjectMacros.h"
 
+#include "DataValidationModule.generated.h"
+
 struct FAssetData;
 
 #include "DataValidationModule.generated.h"
