Avoid reporting a WorldPartitionValidation error when a previous validator failed.

diff --git a/Engine/Plugins/Editor/DataValidation/Source/DataValidation/Private/WorldPartitionChangelistValidator.cpp b/Engine/Plugins/Editor/DataValidation/Source/DataValidation/Private/WorldPartitionChangelistValidator.cpp
--- a/Engine/Plugins/Editor/DataValidation/Source/DataValidation/Private/WorldPartitionChangelistValidator.cpp
+++ b/Engine/Plugins/Editor/DataValidation/Source/DataValidation/Private/WorldPartitionChangelistValidator.cpp
@@ -36,9 +36,9 @@ EDataValidationResult UWorldPartitionChangelistValidator::ValidateLoadedAsset_Im
 	
 	Errors = &ValidationErrors;
 	
-	EDataValidationResult Result = ValidateActorsAndDataLayersFromChangeList(ChangeList);
+	EDataValidationResult ValidateActorAndDataLayerResult = ValidateActorsAndDataLayersFromChangeList(ChangeList);
 
-	if (Result == EDataValidationResult::Invalid)
+	if (ValidateActorAndDataLayerResult == EDataValidationResult::Invalid)
 	{
 		AssetFails(InAsset, LOCTEXT("WorldPartitionValidationFail", "This changelist contains modifications that aren't valid at the world partition level. Please see revision control log and correct the errors."), ValidationErrors);
 	}
@@ -46,8 +46,8 @@ EDataValidationResult UWorldPartitionChangelistValidator::ValidateLoadedAsset_Im
 	{
 		AssetPasses(InAsset);
 	}
-		
-	return Result;
+
+	return Errors->IsEmpty() ? EDataValidationResult::Valid : EDataValidationResult::Invalid;
 }
 
 // Extract all Actors/Map from Changelist (in OFPA this should be one Actor per Package, and we'll discard all Actors from non WorldPartition maps)
@@ -55,6 +55,8 @@ EDataValidationResult UWorldPartitionChangelistValidator::ValidateLoadedAsset_Im
 // from memory (if loaded) or request it to be loaded, we then build a Set of objects that interest us from the Actors in the CL 
 EDataValidationResult UWorldPartitionChangelistValidator::ValidateActorsAndDataLayersFromChangeList(UDataValidationChangelist* Changelist)
 {
+	int32 NumErrorsOnEntry = Errors->Num();
+
 	ISourceControlProvider& SourceControlProvider = ISourceControlModule::Get().GetProvider();
 	FSourceControlChangelistStatePtr ChangelistState = SourceControlProvider.GetState(Changelist->Changelist->AsShared(), EStateCacheUsage::Use);
 
@@ -236,7 +238,7 @@ EDataValidationResult UWorldPartitionChangelistValidator::ValidateActorsAndDataL
 		});
 	}
 
-	if (Errors->Num())
+	if (Errors->Num() > NumErrorsOnEntry)
 	{
 		return EDataValidationResult::Invalid;
 	}
