Add optional per asset details for UEditorValidatorSubsystem::ValidateAssetsWithSettings

diff --git a/Engine/Plugins/Editor/DataValidation/Source/DataValidation/Private/EditorValidatorSubsystem.cpp b/Engine/Plugins/Editor/DataValidation/Source/DataValidation/Private/EditorValidatorSubsystem.cpp
--- a/Engine/Plugins/Editor/DataValidation/Source/DataValidation/Private/EditorValidatorSubsystem.cpp
+++ b/Engine/Plugins/Editor/DataValidation/Source/DataValidation/Private/EditorValidatorSubsystem.cpp
@@ -336,6 +336,16 @@ int32 UEditorValidatorSubsystem::ValidateAssetsWithSettings(const TArray<FAssetD
 				++NumFilesUnableToValidate;
 			}
 		}
+
+		if (InSettings.bCollectPerAssetDetails)
+		{
+			FValidateAssetsDetails& Details = OutResults.AssetsDetails.Emplace(Data.GetObjectPathString());
+			Details.PackageName = Data.PackageName;
+			Details.AssetName = Data.AssetName;
+			Details.Result = Result;
+			Details.ValidationErrors = MoveTemp(ValidationErrors);
+			Details.ValidationWarnings = MoveTemp(ValidationWarnings);
+		}
 	}
 
 	const bool bFailed = (NumInvalidFiles > 0);
diff --git a/Engine/Plugins/Editor/DataValidation/Source/DataValidation/Public/EditorValidatorSubsystem.h b/Engine/Plugins/Editor/DataValidation/Source/DataValidation/Public/EditorValidatorSubsystem.h
--- a/Engine/Plugins/Editor/DataValidation/Source/DataValidation/Public/EditorValidatorSubsystem.h
+++ b/Engine/Plugins/Editor/DataValidation/Source/DataValidation/Public/EditorValidatorSubsystem.h
@@ -18,6 +18,27 @@ typedef TSharedPtr<class ISourceControlChangelist, ESPMode::ThreadSafe> FSourceC
 
 DECLARE_LOG_CATEGORY_EXTERN(LogContentValidation, Log, All);
 
+USTRUCT(BlueprintType)
+struct DATAVALIDATION_API FValidateAssetsDetails
+{
+	GENERATED_BODY()
+
+	/** Package Name */
+	FName PackageName;
+
+	/** Asset Name */
+	FName AssetName;
+
+	/** Validation Result */
+	EDataValidationResult Result = EDataValidationResult::Valid;
+
+	/** Validation Errors */
+	TArray<FText> ValidationErrors;
+
+	/** Validation Warnings */
+	TArray<FText> ValidationWarnings;
+};
+
 USTRUCT(BlueprintType)
 struct DATAVALIDATION_API FValidateAssetsResults
 {
@@ -48,6 +69,14 @@ struct DATAVALIDATION_API FValidateAssetsResults
 	/** Amount of assets that could not be validated */
 	UPROPERTY(VisibleAnywhere, BlueprintReadOnly, Category = "Validation")
 	int NumUnableToValidate = 0;
+
+	/** 
+	 * Per asset details
+	 * Indexed by object path
+	 * Only returned if FValidateAssetsSettings::bCollectPerAssetDetails is true.
+	 */
+	UPROPERTY(VisibleAnywhere, BlueprintReadOnly, Category = "Validation")
+	TMap<FString, FValidateAssetsDetails> AssetsDetails;
 };
 
 USTRUCT(BlueprintType)
@@ -65,6 +94,10 @@ struct DATAVALIDATION_API FValidateAssetsSettings
 	UPROPERTY(EditAnywhere, BlueprintReadWrite, Category = "Validation")
 	bool bShowIfNoFailures = true;
 
+	/** If true, will add an FValidateAssetsDetails for each asset to the results */
+	UPROPERTY(EditAnywhere, BlueprintReadWrite, Category = "Validation")
+	bool bCollectPerAssetDetails = false;
+
 	/** The usecase requiring datavalidation */
 	UPROPERTY(EditAnywhere, BlueprintReadWrite, Category = "Validation")
 	EDataValidationUsecase ValidationUsecase = EDataValidationUsecase::None;
