Changelist Validation: DirtyFilesChangelistValidator now outputs  the assets that need to be saved

diff --git a/Engine/Plugins/Editor/DataValidation/Source/DataValidation/Private/DataValidationChangelist.cpp b/Engine/Plugins/Editor/DataValidation/Source/DataValidation/Private/DataValidationChangelist.cpp
--- a/Engine/Plugins/Editor/DataValidation/Source/DataValidation/Private/DataValidationChangelist.cpp
+++ b/Engine/Plugins/Editor/DataValidation/Source/DataValidation/Private/DataValidationChangelist.cpp
@@ -16,7 +16,7 @@
 
 #define LOCTEXT_NAMESPACE "DataValidationChangelist"
 
-void GatherDependencies(const FName& InPackageName, TSet<FName>& OutDependencies)
+void UDataValidationChangelist::GatherDependencies(const FName& InPackageName, TSet<FName>& OutDependencies)
 {
 	OutDependencies.Add(InPackageName);
 
@@ -34,7 +34,7 @@ void GatherDependencies(const FName& InPackageName, TSet<FName>& OutDependencies
 	}
 }
 
-FString GetPrettyPackageName(const FName& InPackageName)
+FString UDataValidationChangelist::GetPrettyPackageName(const FName& InPackageName)
 {
 	TArray<FAssetData> Assets;
 	USourceControlHelpers::GetAssetDataFromPackage(InPackageName.ToString(), Assets);
diff --git a/Engine/Plugins/Editor/DataValidation/Source/DataValidation/Private/DirtyFilesChangelistValidator.cpp b/Engine/Plugins/Editor/DataValidation/Source/DataValidation/Private/DirtyFilesChangelistValidator.cpp
--- a/Engine/Plugins/Editor/DataValidation/Source/DataValidation/Private/DirtyFilesChangelistValidator.cpp
+++ b/Engine/Plugins/Editor/DataValidation/Source/DataValidation/Private/DirtyFilesChangelistValidator.cpp
@@ -6,28 +6,12 @@
 #include "FileHelpers.h"
 #include "ISourceControlModule.h"
 #include "Misc/Paths.h"
+#include "Containers/Map.h"
 
 #include UE_INLINE_GENERATED_CPP_BY_NAME(DirtyFilesChangelistValidator)
 
 #define LOCTEXT_NAMESPACE "DirtyFilesChangelistValidation"
 
-FString UDirtyFilesChangelistValidator::GetPackagePath(const UPackage* InPackage)
-{
-	if (InPackage == nullptr)
-	{
-		return TEXT("");
-	}
-
-	const FString LocalFullPath(InPackage->GetLoadedPath().GetLocalFullPath());
-
-	if (LocalFullPath.IsEmpty())
-	{
-		return TEXT("");
-	}
-
-	return FPaths::ConvertRelativePathToFull(LocalFullPath);
-}
-
 bool UDirtyFilesChangelistValidator::CanValidateAsset_Implementation(UObject* InAsset) const
 {
 	return (InAsset != nullptr) && (UDataValidationChangelist::StaticClass() == InAsset->GetClass());
@@ -51,23 +35,55 @@ EDataValidationResult UDirtyFilesChangelistValidator::ValidateLoadedAsset_Implem
 	
 	// Retrieve current unsaved packages
 	TArray<UPackage*> DirtyPackages;
-	TSet<FString> DirtyPackagesPath;
+	TMap<FString, const UPackage*> DirtyPackagesPath;
 	FEditorFileUtils::GetDirtyPackages(DirtyPackages, FEditorFileUtils::FShouldIgnorePackage::Default);
 
-	Algo::Transform(DirtyPackages, DirtyPackagesPath, UDirtyFilesChangelistValidator::GetPackagePath);
+	auto Predicate = [](const UPackage* InPackage) -> bool
+	{
+		if (InPackage == nullptr)
+		{
+			return false;
+		}
+		
+		const FString LocalFullPath(InPackage->GetLoadedPath().GetLocalFullPath());
+
+		if (LocalFullPath.IsEmpty())
+		{
+			return false;
+		}
+
+		return true;
+	};
+
+	auto Transform = [](const UPackage* InPackage) -> TTuple<FString, const UPackage*>
+	{
+		check(InPackage);
+		const FString LocalFullPath(InPackage->GetLoadedPath().GetLocalFullPath());
+		check(!LocalFullPath.IsEmpty());
+		
+		return TTuple<FString, const UPackage*>(FPaths::ConvertRelativePathToFull(LocalFullPath), InPackage);
+	};
+		
+	Algo::TransformIf(DirtyPackages, DirtyPackagesPath, Predicate, Transform);
 
 	// Check if any of changelist's files is unsaved
 	for (const FSourceControlStateRef& FileState : FileStates)
 	{
-		if (DirtyPackagesPath.Contains(FileState->GetFilename()))
+		if (const UPackage** PackagePtr = DirtyPackagesPath.Find(FileState->GetFilename()))
 		{
-			AssetFails(InAsset, LOCTEXT("DirtyFilesFound", "This changelist contains unsaved modifications. Please save to proceed."), ValidationErrors);
-			return EDataValidationResult::NotValidated;
+			const UPackage* Package = *PackagePtr;
+			check(Package);
+			FText CurrentError = FText::Format(LOCTEXT("DirtyFilesFound.Changelist.Error", "This changelist contains an unsaved asset {0}. Please save to proceed."), FText::FromString(UDataValidationChangelist::GetPrettyPackageName(Package->GetFName())));
+			AssetFails(InAsset, CurrentError, ValidationErrors);
 		}
 	}
 
-	AssetPasses(InAsset);
-	return EDataValidationResult::Valid;
+	if (GetValidationResult() != EDataValidationResult::Invalid)
+	{
+		AssetPasses(InAsset);
+	}
+
+	return GetValidationResult();
 }
 
 #undef LOCTEXT_NAMESPACE
diff --git a/Engine/Plugins/Editor/DataValidation/Source/DataValidation/Public/DataValidationChangelist.h b/Engine/Plugins/Editor/DataValidation/Source/DataValidation/Public/DataValidationChangelist.h
--- a/Engine/Plugins/Editor/DataValidation/Source/DataValidation/Public/DataValidationChangelist.h
+++ b/Engine/Plugins/Editor/DataValidation/Source/DataValidation/Public/DataValidationChangelist.h
@@ -25,6 +25,8 @@ public:
 	/** Initializes internal state so the validation can be done */
 	void Initialize(FSourceControlChangelistPtr InChangelist);
 
+	static void GatherDependencies(const FName& InPackageName, TSet<FName>& OutDependencies);
+	static FString GetPrettyPackageName(const FName& InPackageName);
 public:
 	/** Changelist to validate */
 	FSourceControlChangelistPtr Changelist;
diff --git a/Engine/Plugins/Editor/DataValidation/Source/DataValidation/Public/DirtyFilesChangelistValidator.h b/Engine/Plugins/Editor/DataValidation/Source/DataValidation/Public/DirtyFilesChangelistValidator.h
--- a/Engine/Plugins/Editor/DataValidation/Source/DataValidation/Public/DirtyFilesChangelistValidator.h
+++ b/Engine/Plugins/Editor/DataValidation/Source/DataValidation/Public/DirtyFilesChangelistValidator.h
@@ -18,7 +18,4 @@ class DATAVALIDATION_API UDirtyFilesChangelistValidator : public UEditorValidato
 protected:
 	virtual bool CanValidateAsset_Implementation(UObject* InAsset) const override;
 	virtual EDataValidationResult ValidateLoadedAsset_Implementation(UObject* InAsset, TArray<FText>& ValidationErrors) override;
-
-private:
-	FString static GetPackagePath(const UPackage* InPackage);
 };
