#UE Adjust validator test for redirectors in localization if both the source and localized assets are redirectors, and they both point to a corresponding new location, with matching relative paths

diff --git a/Engine/Plugins/Editor/DataValidation/Source/DataValidation/Private/EditorValidator_Localization.cpp b/Engine/Plugins/Editor/DataValidation/Source/DataValidation/Private/EditorValidator_Localization.cpp
--- a/Engine/Plugins/Editor/DataValidation/Source/DataValidation/Private/EditorValidator_Localization.cpp
+++ b/Engine/Plugins/Editor/DataValidation/Source/DataValidation/Private/EditorValidator_Localization.cpp
@@ -11,6 +11,7 @@
 #include "HAL/FileManager.h"
 #include "Internationalization/Culture.h"
 #include "Internationalization/PackageLocalizationUtil.h"
+#include "UObject/ConstructorHelpers.h"
 
 #include UE_INLINE_GENERATED_CPP_BY_NAME(EditorValidator_Localization)
 
@@ -59,8 +60,31 @@ EDataValidationResult UEditorValidator_Localization::ValidateLoadedAsset_Impleme
 				// Is the source asset a redirector?
 				if (SourceAssetData.IsRedirector())
 				{
-					AssetFails(InAsset, LOCTEXT("LocalizationError_LocalizedAssetHasSourceRedirector", "Localized asset has a source asset that is a redirector. Did you forget to rename the localized assets too?"), ValidationErrors);
-					return EDataValidationResult::Invalid;
+					// If this asset is a redirector, and the source asset is a redirector, then make sure this asset points to the localized version of what the source asset points to
+					bool bIsRedirectorReferringToRedirectedLocation = false;
+					if (UObjectRedirector* ThisRedirector = Cast<UObjectRedirector>(InAsset))
+					{
+						if (UObject* ThisDestinationObject = ThisRedirector->DestinationObject)
+						{
+							FString ThisDestinationSourceObjectPath;
+							if (ThisDestinationObject->IsLocalizedResource() && FPackageLocalizationUtil::ConvertLocalizedToSource(ThisDestinationObject->GetPathName(), ThisDestinationSourceObjectPath))
+							{
+								FString SourceRedirectorDest;
+								if (SourceAssetData.GetTagValue(TEXT("DestinationObject"), SourceRedirectorDest))
+								{
+									ConstructorHelpers::StripObjectClass(SourceRedirectorDest);
+
+									bIsRedirectorReferringToRedirectedLocation = (ThisDestinationSourceObjectPath == SourceRedirectorDest);
+								}
+							}
+						}
+					}
+
+					if (!bIsRedirectorReferringToRedirectedLocation)
+					{
+						AssetFails(InAsset, LOCTEXT("LocalizationError_LocalizedAssetHasSourceRedirector", "Localized asset has a source asset that is a redirector. Did you forget to rename the localized assets too?"), ValidationErrors);
+						return EDataValidationResult::Invalid;
+					}
 				}
 
 				// Is the source asset the expected type?
@@ -100,10 +124,33 @@ EDataValidationResult UEditorValidator_Localization::ValidateLoadedAsset_Impleme
 							}
 
 							// Is the source asset a redirector?
-							if (InAsset->IsA<UObjectRedirector>())
+							if (UObjectRedirector* ThisRedirector = Cast<UObjectRedirector>(InAsset))
 							{
-								AssetFails(InAsset, FText::Format(LOCTEXT("LocalizationError_SourceIsRedirector", "Source asset is a redirector but has a localized asset for '{0}'. Did you forget to rename the localized assets too?"), FText::FromString(CultureName)), ValidationErrors);
-								return EDataValidationResult::Invalid;
+								// If this source is a redirector, and the localized asset is a redirector, then make sure this asset points to the source version of what the localized asset points to
+								bool bIsRedirectorReferringToRedirectedLocation = false;
+								if (LocalizedAssetData.IsRedirector())
+								{
+									if (UObject* ThisDestinationObject = ThisRedirector->DestinationObject)
+									{
+										FString ThisDestinationLocalizedObjectPath;
+										if (FPackageLocalizationUtil::ConvertSourceToLocalized(ThisDestinationObject->GetPathName(), CultureName, ThisDestinationLocalizedObjectPath))
+										{
+											FString LocalizedRedirectorDest;
+											if (LocalizedAssetData.GetTagValue(TEXT("DestinationObject"), LocalizedRedirectorDest))
+											{
+												ConstructorHelpers::StripObjectClass(LocalizedRedirectorDest);
+
+												bIsRedirectorReferringToRedirectedLocation = (ThisDestinationLocalizedObjectPath == LocalizedRedirectorDest);
+											}
+										}
+									}
+								}
+
+								if (!bIsRedirectorReferringToRedirectedLocation)
+								{
+									AssetFails(InAsset, FText::Format(LOCTEXT("LocalizationError_SourceIsRedirector", "Source asset is a redirector but has a localized asset for '{0}'. Did you forget to rename the localized assets too?"), FText::FromString(CultureName)), ValidationErrors);
+									return EDataValidationResult::Invalid;
+								}
 							}
 
 							// Is the localized asset the expected type?
