Fixed WorldPartition changelist validator containing a data layer asset so that only the referencing partitioned worlds will be validated.

diff --git a/Engine/Plugins/Editor/DataValidation/Source/DataValidation/Private/WorldPartitionChangelistValidator.cpp b/Engine/Plugins/Editor/DataValidation/Source/DataValidation/Private/WorldPartitionChangelistValidator.cpp
--- a/Engine/Plugins/Editor/DataValidation/Source/DataValidation/Private/WorldPartitionChangelistValidator.cpp
+++ b/Engine/Plugins/Editor/DataValidation/Source/DataValidation/Private/WorldPartitionChangelistValidator.cpp
@@ -124,14 +124,17 @@ EDataValidationResult UWorldPartitionChangelistValidator::ValidateActorsAndDataL
 						FARFilter Filter;
 						Filter.bIncludeOnlyOnDiskAssets = true;
 						Filter.PackageNames = MoveTemp(ReferencerNames);
-						Filter.ClassPaths.Add(AWorldDataLayers::StaticClass()->GetClassPathName());
 
 						TArray<FAssetData> DataLayerReferencers;
 						AssetRegistry.GetAssets(Filter, DataLayerReferencers);
 
 						for (const FAssetData& DataLayerReferencer : DataLayerReferencers)
 						{
-							TryAssociateActorToMap(DataLayerReferencer);
+							UClass* ReferencerAssetClass = DataLayerReferencer.GetClass();
+							if (ReferencerAssetClass && ReferencerAssetClass->IsChildOf<AWorldDataLayers>())
+							{
+								TryAssociateActorToMap(DataLayerReferencer);
+							}
 						}
 			
 						RelevantDataLayerAssets.Add(AssetData.PackageName.ToString());
