Added a context menu entry to perform asset validation on specific actors

diff --git a/Engine/Plugins/Editor/DataValidation/Source/DataValidation/DataValidation.Build.cs b/Engine/Plugins/Editor/DataValidation/Source/DataValidation/DataValidation.Build.cs
--- a/Engine/Plugins/Editor/DataValidation/Source/DataValidation/DataValidation.Build.cs
+++ b/Engine/Plugins/Editor/DataValidation/Source/DataValidation/DataValidation.Build.cs
@@ -36,11 +36,13 @@ public class DataValidation : ModuleRules
 				"AssetTools",
 				"PropertyEditor",
 				"GraphEditor",
+				"LevelEditor",
 				"BlueprintGraph",
 				"KismetCompiler",
 				"SandboxFile",
 				"Blutility",
 				"ToolMenus",
+				"TypedElementRuntime",
 				"SourceControl",
 				"UncontrolledChangelists",
 			}
diff --git a/Engine/Plugins/Editor/DataValidation/Source/DataValidation/Private/DataValidationModule.cpp b/Engine/Plugins/Editor/DataValidation/Source/DataValidation/Private/DataValidationModule.cpp
--- a/Engine/Plugins/Editor/DataValidation/Source/DataValidation/Private/DataValidationModule.cpp
+++ b/Engine/Plugins/Editor/DataValidation/Source/DataValidation/Private/DataValidationModule.cpp
@@ -11,8 +11,10 @@
 #include "AssetRegistry/AssetRegistryModule.h"
 #include "Misc/MessageDialog.h"
 #include "DataValidationCommandlet.h"
+#include "Elements/Framework/TypedElementSelectionSet.h"
 
 #include "ContentBrowserMenuContexts.h"
+#include "LevelEditorMenuContext.h"
 
 #include "EditorValidatorSubsystem.h"
 #include "ISettingsModule.h"
@@ -95,6 +97,30 @@ void FDataValidationModule::RegisterMenus()
 		));
 	}
 
+	{
+		UToolMenu* Menu = UToolMenus::Get()->ExtendMenu("LevelEditor.ActorContextMenu");
+		FToolMenuSection& Section = Menu->FindOrAddSection("ActorOptions");
+		Section.AddMenuEntry(
+			"ValidateActors",
+			LOCTEXT("ValidateActorsTabTitle", "Validate"),
+			LOCTEXT("ValidateActorsTooltipText", "Runs data validation on these actors."),
+			FSlateIcon(FAppStyle::GetAppStyleSetName(), "DeveloperTools.MenuIcon"),
+			FToolMenuExecuteAction::CreateLambda([this](const FToolMenuContext& InContext)
+			{
+				if (ULevelEditorContextMenuContext* Context = InContext.FindContext<ULevelEditorContextMenuContext>())
+				{
+					TArray<FAssetData> SelectedActorAssets;
+					Context->CurrentSelection->ForEachSelectedObject<AActor>([&SelectedActorAssets](AActor* SelectedActor)
+					{
+						SelectedActorAssets.Add(FAssetData(SelectedActor));
+						return true;
+					});
+					ValidateAssets(SelectedActorAssets, false, EDataValidationUsecase::Manual);
+				}
+			})
+		);
+	}
+
 	{
 		UToolMenu* Menu = UToolMenus::Get()->ExtendMenu("ContentBrowser.AssetContextMenu.AssetActionsSubMenu");
 		FToolMenuSection& Section = Menu->FindOrAddSection("AssetContextAdvancedActions");
