Only allow the actor validation menu to be used for OFPA packages

#rnx

diff --git a/Engine/Plugins/Editor/DataValidation/Source/DataValidation/Private/DataValidationModule.cpp b/Engine/Plugins/Editor/DataValidation/Source/DataValidation/Private/DataValidationModule.cpp
--- a/Engine/Plugins/Editor/DataValidation/Source/DataValidation/Private/DataValidationModule.cpp
+++ b/Engine/Plugins/Editor/DataValidation/Source/DataValidation/Private/DataValidationModule.cpp
@@ -98,6 +98,37 @@ void FDataValidationModule::RegisterMenus()
 	}
 
 	{
+		FToolUIAction Action;
+		Action.ExecuteAction = FToolMenuExecuteAction::CreateLambda([this](const FToolMenuContext& InContext)
+		{
+			if (ULevelEditorContextMenuContext* Context = InContext.FindContext<ULevelEditorContextMenuContext>())
+			{
+				TArray<FAssetData> SelectedActorAssets;
+				Context->CurrentSelection->ForEachSelectedObject<AActor>([&SelectedActorAssets](AActor* SelectedActor)
+				{
+					if (SelectedActor->GetExternalPackage())
+					{
+						SelectedActorAssets.Add(FAssetData(SelectedActor));
+					}
+					return true;
+				});
+				ValidateAssets(SelectedActorAssets, false, EDataValidationUsecase::Manual);
+			}
+		});
+		Action.CanExecuteAction = FToolMenuCanExecuteAction::CreateLambda([this](const FToolMenuContext& InContext)
+		{
+			bool bCanValidateActor = false;
+			if (ULevelEditorContextMenuContext* Context = InContext.FindContext<ULevelEditorContextMenuContext>())
+			{
+				Context->CurrentSelection->ForEachSelectedObject<AActor>([&bCanValidateActor](AActor* SelectedActor)
+				{
+					bCanValidateActor |= SelectedActor->GetExternalPackage() != nullptr; // Only allow validation of OFPA packages
+					return !bCanValidateActor;
+				});
+			}
+			return bCanValidateActor;
+		});
+
 		UToolMenu* Menu = UToolMenus::Get()->ExtendMenu("LevelEditor.ActorContextMenu");
 		FToolMenuSection& Section = Menu->FindOrAddSection("ActorOptions");
 		Section.AddMenuEntry(
@@ -105,19 +136,7 @@ void FDataValidationModule::RegisterMenus()
 			LOCTEXT("ValidateActorsTabTitle", "Validate"),
 			LOCTEXT("ValidateActorsTooltipText", "Runs data validation on these actors."),
 			FSlateIcon(FAppStyle::GetAppStyleSetName(), "DeveloperTools.MenuIcon"),
-			FToolMenuExecuteAction::CreateLambda([this](const FToolMenuContext& InContext)
-			{
-				if (ULevelEditorContextMenuContext* Context = InContext.FindContext<ULevelEditorContextMenuContext>())
-				{
-					TArray<FAssetData> SelectedActorAssets;
-					Context->CurrentSelection->ForEachSelectedObject<AActor>([&SelectedActorAssets](AActor* SelectedActor)
-					{
-						SelectedActorAssets.Add(FAssetData(SelectedActor));
-						return true;
-					});
-					ValidateAssets(SelectedActorAssets, false, EDataValidationUsecase::Manual);
-				}
-			})
+			Action
 		);
 	}
 
