Extra nullptr check

#rnx

diff --git a/Engine/Plugins/Editor/DataValidation/Source/DataValidation/Private/DataValidationModule.cpp b/Engine/Plugins/Editor/DataValidation/Source/DataValidation/Private/DataValidationModule.cpp
--- a/Engine/Plugins/Editor/DataValidation/Source/DataValidation/Private/DataValidationModule.cpp
+++ b/Engine/Plugins/Editor/DataValidation/Source/DataValidation/Private/DataValidationModule.cpp
@@ -101,7 +101,8 @@ void FDataValidationModule::RegisterMenus()
 		FToolUIAction Action;
 		Action.ExecuteAction = FToolMenuExecuteAction::CreateLambda([this](const FToolMenuContext& InContext)
 		{
-			if (ULevelEditorContextMenuContext* Context = InContext.FindContext<ULevelEditorContextMenuContext>())
+			if (ULevelEditorContextMenuContext* Context = InContext.FindContext<ULevelEditorContextMenuContext>();
+				Context && Context->CurrentSelection)
 			{
 				TArray<FAssetData> SelectedActorAssets;
 				Context->CurrentSelection->ForEachSelectedObject<AActor>([&SelectedActorAssets](AActor* SelectedActor)
@@ -118,7 +119,8 @@ void FDataValidationModule::RegisterMenus()
 		Action.CanExecuteAction = FToolMenuCanExecuteAction::CreateLambda([this](const FToolMenuContext& InContext)
 		{
 			bool bCanValidateActor = false;
-			if (ULevelEditorContextMenuContext* Context = InContext.FindContext<ULevelEditorContextMenuContext>())
+			if (ULevelEditorContextMenuContext* Context = InContext.FindContext<ULevelEditorContextMenuContext>();
+				Context && Context->CurrentSelection)
 			{
 				Context->CurrentSelection->ForEachSelectedObject<AActor>([&bCanValidateActor](AActor* SelectedActor)
 				{
