Changelist Validation: Fix for not calling changelist consistency data validation

`UDataValidationChangelist::IsDataValid` was overriding the deprecated no const version of the function.

Unfortunately, the MSVC compiler does not recognise properly deprecated methods if we have a const and non-const version of a method and the compiler did not report it.

diff --git a/Engine/Plugins/Editor/DataValidation/Source/DataValidation/Private/DataValidationChangelist.cpp b/Engine/Plugins/Editor/DataValidation/Source/DataValidation/Private/DataValidationChangelist.cpp
--- a/Engine/Plugins/Editor/DataValidation/Source/DataValidation/Private/DataValidationChangelist.cpp
+++ b/Engine/Plugins/Editor/DataValidation/Source/DataValidation/Private/DataValidationChangelist.cpp
@@ -69,7 +69,7 @@ FString UDataValidationChangelist::GetPrettyPackageName(const FName& InPackageNa
 	}
 }
 
-EDataValidationResult UDataValidationChangelist::IsDataValid(FDataValidationContext& Context)
+EDataValidationResult UDataValidationChangelist::IsDataValid(FDataValidationContext& Context) const
 {
 	ISourceControlProvider& SourceControlProvider = ISourceControlModule::Get().GetProvider();
 	
diff --git a/Engine/Plugins/Editor/DataValidation/Source/DataValidation/Public/DataValidationChangelist.h b/Engine/Plugins/Editor/DataValidation/Source/DataValidation/Public/DataValidationChangelist.h
--- a/Engine/Plugins/Editor/DataValidation/Source/DataValidation/Public/DataValidationChangelist.h
+++ b/Engine/Plugins/Editor/DataValidation/Source/DataValidation/Public/DataValidationChangelist.h
@@ -20,7 +20,7 @@ public:
 	/** Default constructor with nothing */
 	UDataValidationChangelist() = default;
 
-	virtual EDataValidationResult IsDataValid(class FDataValidationContext& Context) override;
+	virtual EDataValidationResult IsDataValid(class FDataValidationContext& Context) const override;
 
 	/** Initializes internal state so the validation can be done */
 	void Initialize(FSourceControlChangelistPtr InChangelist);
