DataValidation refactor

General
    Add support for returning tokenized messages from validation.
    Unify validation reporting for different validation entry points - in particular, source control window context menu.
    Move EDataValidationUsecase to CoreUObject for use by UObject::IsDataValid.

UEditorValidatorSubsystem
    Add support for validating changelists. Use of this interface allows building lists of unmodified assets to re-validate based on the changes that were made.
    Allow subclassing, with functions GatherAssetsToValidateFromChangelist and ShouldValidateAsset to modify behavior of validation.
    Add optional to load assets for validation to validation parameters.
    Capture logs during asset loads to report as errors (optional in validation parameters).
    Allow limit to number of assets to validate in validation parameters, report limit reached.
    Add message log page to validation parameters.
    Add validation parameter to validate referencers of deleted assets in changelists for validation

UEditorValidatorBase
    Centralize dispatch to BP/Native validation in non-virtual ValidateLoadedAsset function.
    Add native interface functions to UEditorValidatorBase that take FDataValidationContext.
    Rename BP interface functions to have K2 prefix and different signature to native functions.

FDataValidationContext
    Add validation usecase.
    Add 'associated assets' to allow associating external actors with maps for validators.

UDataValidationChangelist
    Support changelists based on changelist objects, file states, and simple lists of files.
    Group modifications into packages/non-assets and modified/deleted ahead of time for validators to inspect.

diff --git a/Engine/Config/BaseEngine.ini b/Engine/Config/BaseEngine.ini
--- a/Engine/Config/BaseEngine.ini
+++ b/Engine/Config/BaseEngine.ini
@@ -1212,6 +1212,12 @@ DefaultWorldPartitionHLODLayer=/Engine/Maps/Templates/HLODs/HLODLayer_Instanced.
 +ClassRedirects=(OldName="/Script/Engine.MaterialExpressionStrataThinFilm",						NewName="/Script/Engine.MaterialExpressionSubstrateThinFilm")
 +StructRedirects=(OldName="/Script/Engine.StrataMaterialInput",									NewName="/Script/Engine.SubstrateMaterialInput")
 
+; Data Validation refactor 
++EnumRedirects=(OldName="/Script/DataValidation.EDataValidationUsecase",NewName="/Script/CoreUObject.EDataValidationUsecase")
++FunctionRedirects=(OldName="EditorValidatorBase.ValidateLoadedAsset", NewName="EditorValidatorBase.K2_ValidateLoadedAsset")
++FunctionRedirects=(OldName="EditorValidatorBase.CanValidate", NewName="EditorValidatorBase.K2_CanValidate")
++FunctionRedirects=(OldName="EditorValidatorBase.CanValidateAsset", NewName="EditorValidatorBase.K2_CanValidateAsset")
+
 ; End of releases redirects
 
 
diff --git a/Engine/Plugins/Editor/AssetReferenceRestrictions/Source/AssetReferenceRestrictions/Private/AssetValidator_AssetReferenceRestrictions.cpp b/Engine/Plugins/Editor/AssetReferenceRestrictions/Source/AssetReferenceRestrictions/Private/AssetValidator_AssetReferenceRestrictions.cpp
--- a/Engine/Plugins/Editor/AssetReferenceRestrictions/Source/AssetReferenceRestrictions/Private/AssetValidator_AssetReferenceRestrictions.cpp
+++ b/Engine/Plugins/Editor/AssetReferenceRestrictions/Source/AssetReferenceRestrictions/Private/AssetValidator_AssetReferenceRestrictions.cpp
@@ -1,6 +1,9 @@
 // Copyright Epic Games, Inc. All Rights Reserved.
 
 #include "AssetValidator_AssetReferenceRestrictions.h"
+
+#include "AssetRegistry/AssetData.h"
+#include "AssetRegistry/AssetDataToken.h"
 #include "AssetRegistry/AssetRegistryModule.h"
 #include "Editor.h"
 #include "AssetReferencingPolicySubsystem.h"
@@ -19,11 +22,10 @@ UAssetValidator_AssetReferenceRestrictions::UAssetValidator_AssetReferenceRestri
 {
 }
 
-bool UAssetValidator_AssetReferenceRestrictions::CanValidateAsset_Implementation(UObject* InAsset) const
+bool UAssetValidator_AssetReferenceRestrictions::CanValidateAsset_Implementation(const FAssetData& AssetData, UObject* InAsset, FDataValidationContext& InContext) const
 {
 	if (InAsset)
 	{
-		const FAssetData AssetData(InAsset);
 		TSharedPtr<FDomainData> DomainData = GEditor->GetEditorSubsystem<UAssetReferencingPolicySubsystem>()->GetDomainDB()->FindDomainFromAssetData(AssetData);
 
 		const bool bIsInUnrestrictedFolder = DomainData && DomainData->IsValid() && DomainData->bCanSeeEverything;
@@ -36,7 +38,7 @@ bool UAssetValidator_AssetReferenceRestrictions::CanValidateAsset_Implementation
 	return false;
 }
 
-EDataValidationResult UAssetValidator_AssetReferenceRestrictions::ValidateLoadedAsset_Implementation(UObject* InAsset, TArray<FText>& ValidationErrors)
+EDataValidationResult UAssetValidator_AssetReferenceRestrictions::ValidateLoadedAsset_Implementation(const FAssetData& InAssetData, UObject* InAsset, FDataValidationContext& InContext)
 {
 	check(InAsset);
 
@@ -70,7 +72,7 @@ EDataValidationResult UAssetValidator_AssetReferenceRestrictions::ValidateLoaded
 			{
 				if (SoftDependency != TransientName)
 				{
-					AssetFails(InAsset, FText::Format(LOCTEXT("IllegalReference_MissingSoftRef", "Soft references {0} which does not exist"), FText::FromString(SoftDependencyStr)), ValidationErrors);
+					AssetFails(InAsset, FText::Format(LOCTEXT("IllegalReference_MissingSoftRef", "Soft references {0} which does not exist"), FText::FromString(SoftDependencyStr)));
 				}
 			}
 			else
@@ -91,7 +93,7 @@ EDataValidationResult UAssetValidator_AssetReferenceRestrictions::ValidateLoaded
 		FString UncookedFolderName;
 		if (IsInUncookedFolder(HardDependencyStr, &UncookedFolderName))
 		{
-			AssetFails(InAsset, FText::Format(LOCTEXT("IllegalReference_HardDependency", "Illegally hard references {0} asset {1}"), FText::FromString(UncookedFolderName), FText::FromString(HardDependencyStr)), ValidationErrors);
+			AssetFails(InAsset, FText::Format(LOCTEXT("IllegalReference_HardDependency", "Illegally hard references {0} asset {1}"), FText::FromString(UncookedFolderName), FText::FromString(HardDependencyStr)));
 		}
 #endif
 
@@ -101,7 +103,7 @@ EDataValidationResult UAssetValidator_AssetReferenceRestrictions::ValidateLoaded
 	if ((GetValidationResult() != EDataValidationResult::Invalid) && (AllDependencyAssets.Num() > 0))
 	{
 		FAssetReferenceFilterContext AssetReferenceFilterContext;
-		AssetReferenceFilterContext.ReferencingAssets = { FAssetData(InAsset) };
+		AssetReferenceFilterContext.ReferencingAssets = { InAssetData };
 		TSharedPtr<IAssetReferenceFilter> AssetReferenceFilter = GEditor ? GEditor->MakeAssetReferenceFilter(AssetReferenceFilterContext) : nullptr;
 		if (ensure(AssetReferenceFilter.IsValid()))
 		{
@@ -110,7 +112,9 @@ EDataValidationResult UAssetValidator_AssetReferenceRestrictions::ValidateLoaded
 				FText FailureReason;
 				if (!AssetReferenceFilter->PassesFilter(Dependency, &FailureReason))
 				{
-					AssetFails(InAsset, FText::Format(LOCTEXT("IllegalReference_AssetFilterFail", "Illegally references asset {0}. {1}"), FText::FromName(Dependency.PackageName), FailureReason), ValidationErrors);
+					AssetMessage(InAsset, EMessageSeverity::Error, FText::Format(LOCTEXT("IllegalReference_AssetFilterFail", "Illegal reference:"), FailureReason))
+						->AddToken(FAssetDataToken::Create(Dependency))
+						->AddText(FText::Format(LOCTEXT("IllegalReference_FailureReason", ". {0}"), FailureReason));
 				}
 			}
 		}
diff --git a/Engine/Plugins/Editor/AssetReferenceRestrictions/Source/AssetReferenceRestrictions/Private/AssetValidator_AssetReferenceRestrictions.h b/Engine/Plugins/Editor/AssetReferenceRestrictions/Source/AssetReferenceRestrictions/Private/AssetValidator_AssetReferenceRestrictions.h
--- a/Engine/Plugins/Editor/AssetReferenceRestrictions/Source/AssetReferenceRestrictions/Private/AssetValidator_AssetReferenceRestrictions.h
+++ b/Engine/Plugins/Editor/AssetReferenceRestrictions/Source/AssetReferenceRestrictions/Private/AssetValidator_AssetReferenceRestrictions.h
@@ -15,7 +15,7 @@ public:
 
 protected:
 	//~UEditorValidatorBase interface
-	virtual bool CanValidateAsset_Implementation(UObject* InAsset) const override;
-	virtual EDataValidationResult ValidateLoadedAsset_Implementation(UObject* InAsset, TArray<FText>& ValidationErrors) override;
+	virtual bool CanValidateAsset_Implementation(const FAssetData& InAssetData, UObject* InAsset, FDataValidationContext& InContext) const override;
+	virtual EDataValidationResult ValidateLoadedAsset_Implementation(const FAssetData& InAssetData, UObject* InAsset, FDataValidationContext& InContext) override;
 	//~End of UEditorValidatorBase interface
 };
diff --git a/Engine/Plugins/Editor/DataValidation/Source/DataValidation/DataValidation.Build.cs b/Engine/Plugins/Editor/DataValidation/Source/DataValidation/DataValidation.Build.cs
--- a/Engine/Plugins/Editor/DataValidation/Source/DataValidation/DataValidation.Build.cs
+++ b/Engine/Plugins/Editor/DataValidation/Source/DataValidation/DataValidation.Build.cs
@@ -17,7 +17,8 @@ public class DataValidation : ModuleRules
 				"CoreUObject",
 				"Engine",
 				"TargetPlatform",
-				"EditorSubsystem"
+				"EditorSubsystem",
+				"DeveloperSettings"
 			}
 		);
 //		goto through these one by one and remove extra ones
diff --git a/Engine/Plugins/Editor/DataValidation/Source/DataValidation/Private/DataValidationChangelist.cpp b/Engine/Plugins/Editor/DataValidation/Source/DataValidation/Private/DataValidationChangelist.cpp
--- a/Engine/Plugins/Editor/DataValidation/Source/DataValidation/Private/DataValidationChangelist.cpp
+++ b/Engine/Plugins/Editor/DataValidation/Source/DataValidation/Private/DataValidationChangelist.cpp
@@ -76,27 +76,15 @@ FString UDataValidationChangelist::GetPrettyPackageName(const FName& InPackageNa
 EDataValidationResult UDataValidationChangelist::IsDataValid(FDataValidationContext& Context) const
 {
 	ISourceControlProvider& SourceControlProvider = ISourceControlModule::Get().GetProvider();
-	
-	FSourceControlChangelistStatePtr ChangelistState = SourceControlProvider.GetState(Changelist->AsShared(), EStateCacheUsage::Use);
 
 	// Gather dependencies of every file in the changelist
 	TArray<FName> FilesInChangelist;
 	TSet<FName> ExternalDependenciesSet;
 	
-	for (const FSourceControlStateRef& File : ChangelistState->GetFilesStates())
+	for (FName PackageName : ModifiedPackageNames)
 	{
-		// We shouldn't consider dependencies of deleted files
-		if (File->IsDeleted())
-		{
-			continue;
-		}
-
-		FString PackageName;
-		if (FPackageName::TryConvertFilenameToLongPackageName(File->GetFilename(), PackageName))
-		{
-			FilesInChangelist.Add(*PackageName);
-			GatherDependencies(*PackageName, ExternalDependenciesSet);
-		}
+		FilesInChangelist.Add(PackageName);
+		GatherDependencies(PackageName, ExternalDependenciesSet);
 	}
 
 	// For every dependency in the external dependencies that is not in the changelist
@@ -190,7 +178,58 @@ EDataValidationResult UDataValidationChangelist::IsDataValid(FDataValidationCont
 
 void UDataValidationChangelist::Initialize(FSourceControlChangelistPtr InChangelist)
 {
-	Changelist = InChangelist;
+	TRACE_CPUPROFILER_EVENT_SCOPE_STR("UDataValidationChangelist::Initialize");
+	if (InChangelist.IsValid())
+	{
+		ISourceControlProvider& Provider = ISourceControlModule::Get().GetProvider();
+		FSourceControlChangelistStatePtr ChangelistStatePtr = Provider.GetState(InChangelist.ToSharedRef(), EStateCacheUsage::Use);
+		if (ChangelistStatePtr.IsValid())
+		{
+			Initialize(ChangelistStatePtr.ToSharedRef());
+		}
+	}
+}
+
+void UDataValidationChangelist::Initialize(FSourceControlChangelistStateRef InChangelist)
+{
+	Changelist = InChangelist->GetChangelist();	
+	Description = InChangelist->GetDescriptionText();
+	Initialize(InChangelist->GetFilesStates());
+}
+
+void UDataValidationChangelist::Initialize(TConstArrayView<FSourceControlStateRef> FileStates)
+{
+	for (const FSourceControlStateRef& FileState : FileStates)
+	{
+		FString Filename = FileState->GetFilename();	
+		if (FPackageName::IsPackageFilename(Filename))
+		{
+			FString PackageName;
+			FString FailureReason;
+			if (FPackageName::TryConvertFilenameToLongPackageName(Filename, PackageName, &FailureReason))
+			{
+				if (FileState->IsDeleted())
+				{
+					DeletedPackageNames.Add(FName(PackageName));
+				}
+				else
+				{
+					ModifiedPackageNames.Add(FName(PackageName));
+				}
+			}
+			continue;
+		}
+		
+		// If it wasn't an asset file or we failed to find the relevant package name 
+		if (FileState->IsDeleted())
+		{
+			DeletedFiles.Add(MoveTemp(Filename));
+		}
+		else 
+		{
+			ModifiedFiles.Add(MoveTemp(Filename));
+		}
+	}
 }
 
 #undef LOCTEXT_NAMESPACE
diff --git a/Engine/Plugins/Editor/DataValidation/Source/DataValidation/Private/DataValidationModule.cpp b/Engine/Plugins/Editor/DataValidation/Source/DataValidation/Private/DataValidationModule.cpp
--- a/Engine/Plugins/Editor/DataValidation/Source/DataValidation/Private/DataValidationModule.cpp
+++ b/Engine/Plugins/Editor/DataValidation/Source/DataValidation/Private/DataValidationModule.cpp
@@ -303,6 +303,7 @@ void FDataValidationModule::ValidateAssets(const TArray<FAssetData>& SelectedAss
 		Settings.bSkipExcludedDirectories = false;
 		Settings.bShowIfNoFailures = true;
 		Settings.ValidationUsecase = InValidationUsecase;
+		Settings.MessageLogPageTitle = LOCTEXT("ValidateSelectedAssets", "Validate Selected Assets");
 
 		EditorValidationSubsystem->ValidateAssetsWithSettings(bValidateDependencies ? DependentAssets.Array() : SelectedAssets, Settings, Results);
 	}
diff --git a/Engine/Plugins/Editor/DataValidation/Source/DataValidation/Private/DirtyFilesChangelistValidator.cpp b/Engine/Plugins/Editor/DataValidation/Source/DataValidation/Private/DirtyFilesChangelistValidator.cpp
--- a/Engine/Plugins/Editor/DataValidation/Source/DataValidation/Private/DirtyFilesChangelistValidator.cpp
+++ b/Engine/Plugins/Editor/DataValidation/Source/DataValidation/Private/DirtyFilesChangelistValidator.cpp
@@ -2,6 +2,9 @@
 
 #include "DirtyFilesChangelistValidator.h"
 
+#include "AssetRegistry/AssetData.h"
+#include "AssetRegistry/AssetDataToken.h"
+#include "AssetRegistry/AssetRegistryModule.h"
 #include "DataValidationChangelist.h"
 #include "FileHelpers.h"
 #include "ISourceControlModule.h"
@@ -12,30 +15,17 @@
 
 #define LOCTEXT_NAMESPACE "DirtyFilesChangelistValidation"
 
-bool UDirtyFilesChangelistValidator::CanValidateAsset_Implementation(UObject* InAsset) const
+bool UDirtyFilesChangelistValidator::CanValidateAsset_Implementation(const FAssetData& AssetData, UObject* InAsset, FDataValidationContext& InContext) const
 {
 	return (InAsset != nullptr) && (UDataValidationChangelist::StaticClass() == InAsset->GetClass());
 }
 
-EDataValidationResult UDirtyFilesChangelistValidator::ValidateLoadedAsset_Implementation(UObject* InAsset, TArray<FText>& ValidationErrors)
+EDataValidationResult UDirtyFilesChangelistValidator::ValidateLoadedAsset_Implementation(const FAssetData& AssetData, UObject* InAsset, FDataValidationContext& InContext)
 {
-	check(UDataValidationChangelist::StaticClass() == InAsset->GetClass());
+	UDataValidationChangelist* DataValidationChangelist = CastChecked<UDataValidationChangelist>(InAsset);
 
-	UDataValidationChangelist* DataValidationChangelist = Cast<UDataValidationChangelist>(InAsset);
-
-	if ((DataValidationChangelist == nullptr) || (!DataValidationChangelist->Changelist.IsValid()))
-	{
-		return EDataValidationResult::Valid;
-	}
-
-	// Retrieve files contained in the changelist
-	ISourceControlProvider& Provider = ISourceControlModule::Get().GetProvider();
-	FSourceControlChangelistStatePtr ChangelistStatePtr = Provider.GetState(DataValidationChangelist->Changelist.ToSharedRef(), EStateCacheUsage::Use);
-	const TArray<FSourceControlStateRef>& FileStates = ChangelistStatePtr->GetFilesStates();
-	
 	// Retrieve current unsaved packages
 	TArray<UPackage*> DirtyPackages;
-	TMap<FString, const UPackage*> DirtyPackagesPath;
 	FEditorFileUtils::GetDirtyPackages(DirtyPackages, FEditorFileUtils::FShouldIgnorePackage::Default);
 
 	auto Predicate = [](const UPackage* InPackage) -> bool
@@ -55,26 +45,28 @@ EDataValidationResult UDirtyFilesChangelistValidator::ValidateLoadedAsset_Implem
 		return true;
 	};
 
-	auto Transform = [](const UPackage* InPackage) -> TTuple<FString, const UPackage*>
-	{
-		check(InPackage);
-		const FString LocalFullPath(InPackage->GetLoadedPath().GetLocalFullPath());
-		check(!LocalFullPath.IsEmpty());
-		
-		return TTuple<FString, const UPackage*>(FPaths::ConvertRelativePathToFull(LocalFullPath), InPackage);
-	};
-		
-	Algo::TransformIf(DirtyPackages, DirtyPackagesPath, Predicate, Transform);
+	TSet<FName> DirtyPackagesSet;
+	Algo::TransformIf(DirtyPackages, DirtyPackagesSet, Predicate, UE_PROJECTION_MEMBER(UPackage, GetFName));
+
+	FAssetRegistryModule& AssetRegistryModule = FModuleManager::LoadModuleChecked<FAssetRegistryModule>("AssetRegistry");
+    IAssetRegistry& AssetRegistry = AssetRegistryModule.Get();
 
 	// Check if any of changelist's files is unsaved
-	for (const FSourceControlStateRef& FileState : FileStates)
+	for (FName PackageName : DataValidationChangelist->ModifiedPackageNames)
 	{
-		if (const UPackage** PackagePtr = DirtyPackagesPath.Find(FileState->GetFilename()))
+		if (DirtyPackagesSet.Contains(PackageName))
 		{
-			const UPackage* Package = *PackagePtr;
-			check(Package);
-			FText CurrentError = FText::Format(LOCTEXT("DirtyFilesFound.Changelist.Error", "This changelist contains an unsaved asset {0}. Please save to proceed."), FText::FromString(UDataValidationChangelist::GetPrettyPackageName(Package->GetFName())));
-			AssetFails(InAsset, CurrentError, ValidationErrors);
+			FText CurrentError = LOCTEXT("DirtyFilesFound.Changelist.Error", "This changelist contains an unsaved asset. Please save to proceed.");
+			TSharedRef<FTokenizedMessage> Message = AssetMessage(EMessageSeverity::Error, CurrentError);
+			TArray<FAssetData> PackageAssets;
+			if (AssetRegistry.GetAssetsByPackageName(PackageName, PackageAssets, true))
+			{
+				Message->AddToken(FAssetDataToken::Create(PackageAssets[0]));
+			}
+			else
+			{
+				Message->AddText(FText::FromName(PackageName));
+			}
 		}
 	}
 
diff --git a/Engine/Plugins/Editor/DataValidation/Source/DataValidation/Private/EditorValidatorBase.cpp b/Engine/Plugins/Editor/DataValidation/Source/DataValidation/Private/EditorValidatorBase.cpp
--- a/Engine/Plugins/Editor/DataValidation/Source/DataValidation/Private/EditorValidatorBase.cpp
+++ b/Engine/Plugins/Editor/DataValidation/Source/DataValidation/Private/EditorValidatorBase.cpp
@@ -2,7 +2,10 @@
 
 #include "EditorValidatorBase.h"
 
+#include "AssetRegistry/AssetData.h"
+#include "AssetRegistry/AssetDataToken.h"
 #include "EditorValidatorSubsystem.h"
+#include "Misc/DataValidation.h"
 #include "Trace/Trace.inl"
 
 #include UE_INLINE_GENERATED_CPP_BY_NAME(EditorValidatorBase)
@@ -15,22 +18,76 @@ UEditorValidatorBase::UEditorValidatorBase()
 	bIsEnabled = true;
 }
 
-bool UEditorValidatorBase::CanValidate_Implementation(const EDataValidationUsecase InUsecase) const
+EDataValidationResult UEditorValidatorBase::ValidateLoadedAsset(const FAssetData& AssetData, UObject* Asset, FDataValidationContext& Context)
 {
-	return true;
+	EDataValidationResult Result = EDataValidationResult::NotValidated;
+	
+	ResetValidationState();
+	
+	TGuardValue<UObject*> ObjectGuard(CurrentObjectBeingValidated, Asset);
+	TGuardValue<const FAssetData*> AssetGuard(CurrentAssetBeingValidated, &AssetData);
+
+	if (K2_CanValidate(Context.GetValidationUsecase()))
+	{
+		EDataValidationResult K2Result = K2_ValidateLoadedAsset(Asset);
+		K2Result = CombineDataValidationResults(GetValidationResult(), K2Result);
+		ensureMsgf(K2Result != EDataValidationResult::NotValidated, TEXT("Validator %s did not return a validation result from BP validation path for asset %s"), *GetClass()->GetPathName(), *Asset->GetPathName());
+	}
+	
+PRAGMA_DISABLE_DEPRECATION_WARNINGS
+	if (CanValidate_Implementation(Context.GetValidationUsecase()) && CanValidateAsset_Implementation(Asset))
+	{
+		ensureMsgf(false, TEXT("Validator %s running deprecated validation path. This validator should implement ValidateLoadedAsset_Implementation(UObject*, FDataValidationContext&) instead"),
+			*GetClass()->GetPathName());
+		
+		EDataValidationResult NewResult = ValidateLoadedAsset_Implementation(Asset, AllErrors);
+		NewResult = CombineDataValidationResults(GetValidationResult(), NewResult);
+		ensureMsgf(NewResult != EDataValidationResult::NotValidated, TEXT("Validator %s did not return a validation result from legacy validation path for asset %s"), *GetClass()->GetPathName(), *Asset->GetPathName());
+
+		Result = CombineDataValidationResults(Result, NewResult);
+	}
+PRAGMA_ENABLE_DEPRECATION_WARNINGS
+	
+	if (CanValidateAsset_Implementation(AssetData, Asset, Context))
+	{
+		EDataValidationResult NewResult = ValidateLoadedAsset_Implementation(AssetData, Asset, Context);
+
+		NewResult = CombineDataValidationResults(GetValidationResult(), NewResult);
+		ensureMsgf(NewResult != EDataValidationResult::NotValidated, TEXT("Validator %s did not return a validation result from native validation path for asset %s"), *GetClass()->GetPathName(), *Asset->GetPathName());
+
+		Result = CombineDataValidationResults(Result, NewResult);
+	}
+	
+	Result = CombineDataValidationResults(Result, ExtractValidationState(Context)); // Extract messages and validation state from members (primarily for use by BP interface functions)
+
+	return Result;
+}
+
+bool UEditorValidatorBase::K2_CanValidate_Implementation(const EDataValidationUsecase InUsecase) const
+{
+	return false;
 }
 
-bool UEditorValidatorBase::CanValidateAsset_Implementation(UObject* InAsset) const
+bool UEditorValidatorBase::K2_CanValidateAsset_Implementation(UObject* InAsset) const
 {
 	return false;
 }
 
-EDataValidationResult UEditorValidatorBase::ValidateLoadedAsset_Implementation(UObject* InAsset, TArray<FText>& ValidationErrors)
+EDataValidationResult UEditorValidatorBase::K2_ValidateLoadedAsset_Implementation(UObject* InAsset)
 {
 	return EDataValidationResult::NotValidated;
 }
 
-void UEditorValidatorBase::AssetFails(UObject* InAsset, const FText& InMessage, TArray<FText>& ValidationErrors)
+void UEditorValidatorBase::AssetFails(UObject* InAsset, const FText& InMessage, TArray<FText>& InOutErrors)
+{
+	AssetFails(InAsset, InMessage);
+	if (&InOutErrors != &AllErrors)
+	{
+		InOutErrors.Add(InMessage);
+	}
+}
+
+void UEditorValidatorBase::AssetFails(UObject* InAsset, const FText& InMessage)
 {
 	FFormatNamedArguments Arguments;
 	Arguments.Add(TEXT("CustomMessage"), InMessage);
@@ -44,7 +101,7 @@ void UEditorValidatorBase::AssetFails(UObject* InAsset, const FText& InMessage,
 
 	}
 
-	ValidationErrors.Add(FailureMessage);
+	AllErrors.Add(FailureMessage);
 	ValidationResult = EDataValidationResult::Invalid;
 }
 
@@ -72,7 +129,6 @@ void UEditorValidatorBase::LogElapsedTime(FFormatNamedArguments &Arguments)
 
 void UEditorValidatorBase::AssetPasses(UObject* InAsset)
 {
-
 	if (LogContentValidation.GetVerbosity() >= ELogVerbosity::Verbose)
 	{
 		FFormatNamedArguments Arguments;
@@ -80,19 +136,68 @@ void UEditorValidatorBase::AssetPasses(UObject* InAsset)
 		{
 			Arguments.Add(TEXT("AssetName"), FText::FromName(InAsset->GetFName()));
 		}
-		Arguments.Add(TEXT("ValidatorName"), FText::FromString(GetClass()->GetName()));
+		Arguments.Add(TEXT("ValidatorName"), FText::FromString(GetClass()->GetPathName()));
 
 		LogElapsedTime(Arguments);
 	}
 
+	ensureMsgf(ValidationResult != EDataValidationResult::Invalid, TEXT("%s: AssetPasses called after errors were reported"), *GetClass()->GetPathName());
 	ValidationResult = EDataValidationResult::Valid;
 }
 
+TSharedRef<FTokenizedMessage> UEditorValidatorBase::AssetMessage(const FAssetData& InAssetData, EMessageSeverity::Type InSeverity, const FText& InText)
+{
+	if (InSeverity == EMessageSeverity::Error)
+	{
+		ValidationResult = EDataValidationResult::Invalid;
+	}
+	TSharedRef<FTokenizedMessage> Message = FTokenizedMessage::Create(InSeverity);
+	if (InAssetData.IsValid())
+	{
+		Message->AddToken(FAssetDataToken::Create(InAssetData));		
+	}
+	if (!InText.IsEmpty())
+	{
+		Message->AddText(InText);
+	}
+	AllMessages.Add(Message);
+	return Message; 
+}
+
+TSharedRef<FTokenizedMessage> UEditorValidatorBase::AssetMessage(EMessageSeverity::Type InSeverity, const FText& InText)
+{
+	return AssetMessage(FAssetData{}, InSeverity, InText);
+}
+
+void UEditorValidatorBase::AddLegacyValidationErrors(TArray<FText> InErrors)
+{
+	AllErrors.Append(MoveTemp(InErrors));
+}
+
 void UEditorValidatorBase::ResetValidationState()
 {
 	ValidationResult = EDataValidationResult::NotValidated;
 	ValidationTime = FDateTime::Now();
+	AllMessages.Empty();
 	AllWarnings.Empty();
+	AllErrors.Empty();
+}
+
+EDataValidationResult UEditorValidatorBase::ExtractValidationState(FDataValidationContext& InOutContext) const
+{
+	for (const FText& Warning : AllWarnings)
+	{
+		InOutContext.AddWarning(Warning);	
+	}
+	for (const FText& Error : AllErrors)
+	{
+		InOutContext.AddError(Error);	
+	}
+	for (const TSharedRef<FTokenizedMessage>& Message : AllMessages)
+	{
+		InOutContext.AddMessage(Message);	
+	}
+	return ValidationResult;
 }
 
 const TArray<FText>& UEditorValidatorBase::GetAllWarnings() const
diff --git a/Engine/Plugins/Editor/DataValidation/Source/DataValidation/Private/EditorValidatorHelpers.cpp b/Engine/Plugins/Editor/DataValidation/Source/DataValidation/Private/EditorValidatorHelpers.cpp
new file mode 100644
--- /dev/null
+++ b/Engine/Plugins/Editor/DataValidation/Source/DataValidation/Private/EditorValidatorHelpers.cpp
@@ -0,0 +1,290 @@
+// Copyright Epic Games, Inc. All Rights Reserved.
+
+#include "EditorValidatorHelpers.h"
+
+#include "Algo/AllOf.h"
+#include "AssetRegistry/AssetData.h"
+#include "AssetRegistry/AssetDataToken.h"
+#include "Logging/MessageLog.h"
+#include "Containers/ArrayView.h"
+#include "Misc/DataValidation.h"
+#include "Misc/WildcardString.h"
+
+namespace UE::DataValidation
+{
+
+const FName MessageLogName("AssetCheck");
+
+void AddAssetValidationMessages(
+    FMessageLog& Log,
+    const FDataValidationContext& ValidationContext)
+{
+    for (const FDataValidationContext::FIssue& Issue : ValidationContext.GetIssues())
+    {
+        if (Issue.TokenizedMessage.IsValid())
+        {
+            Log.AddMessage(Issue.TokenizedMessage.ToSharedRef());
+        }
+        else
+        {
+            Log.Message(Issue.Severity, Issue.Message);
+        }
+    }
+}
+
+void AddAssetValidationMessages(
+    const FAssetData& ForAsset,
+    FMessageLog& Log,
+    const FDataValidationContext& ValidationContext)
+{
+    for (const FDataValidationContext::FIssue& Issue : ValidationContext.GetIssues())
+    {
+        if (Issue.TokenizedMessage.IsValid())
+        {
+            Log.AddMessage(Issue.TokenizedMessage.ToSharedRef());
+        }
+        else
+        {
+            AddAssetValidationMessages(ForAsset, Log, Issue.Severity, { Issue.Message });
+        }
+    }
+}
+
+void AddAssetValidationMessages(
+    const FAssetData& ForAsset,
+    FMessageLog& Log,
+    EMessageSeverity::Type Severity,
+    TConstArrayView<FText> Messages)
+{
+    const FString PackageNameString = ForAsset.PackageName.ToString();
+    for (const FText& Msg : Messages)
+    {
+        const FString AssetLogString = FAssetMsg::GetAssetLogString(*PackageNameString, Msg.ToString());
+        FString BeforeAsset;
+        FString AfterAsset;
+        TSharedRef<FTokenizedMessage> TokenizedMessage = Log.Message(Severity);
+        if (AssetLogString.Split(PackageNameString, &BeforeAsset, &AfterAsset))
+        {
+            if (!BeforeAsset.IsEmpty())
+            {
+                TokenizedMessage->AddToken(FTextToken::Create(FText::FromString(BeforeAsset)));
+            }
+            TokenizedMessage->AddToken(FAssetDataToken::Create(ForAsset));
+            if (!AfterAsset.IsEmpty())
+            {
+                TokenizedMessage->AddToken(FTextToken::Create(FText::FromString(AfterAsset)));
+            }
+        }
+        else
+        {
+            TokenizedMessage->AddToken(FTextToken::Create(FText::FromString(AssetLogString)));
+        }
+    }
+}
+
+struct FLogMessageGathererImpl : public FOutputDevice
+{
+    FLogMessageGathererImpl()
+    {
+        TDelegate<TArray<FName>(void)>& DefaultIgnoreCategoriesDelegate = FScopedLogMessageGatherer::GetDefaultIgnoreCategoriesDelegate();
+        if (DefaultIgnoreCategoriesDelegate.IsBound())
+        {
+            DefaultIgnoreCategories = DefaultIgnoreCategoriesDelegate.Execute();
+        }
+        TDelegate<TArray<FWildcardString>(void)>& DefaultIgnorePatternsDelegate = FScopedLogMessageGatherer::GetDefaultIgnorePatternsDelegate();
+        if (DefaultIgnorePatternsDelegate.IsBound())
+        {
+            DefaultIgnorePatterns = DefaultIgnorePatternsDelegate.Execute();
+        }
+
+        GLog->AddOutputDevice(this);
+    }
+    
+    ~FLogMessageGathererImpl()
+    {
+        // Lock the impl, remove it as an output device, and delete it 
+        // There should be no problems with anyone waiting for the lock as access is controlled by the FScopedLogMessageGatherer
+        // and we are deleted when that goes out of scope
+        CS.Lock();
+        GLog->RemoveOutputDevice(this);
+    }
+   
+    TTuple<int32, int32> AddIgnoreCategories(TConstArrayView<FName> NewCategories)
+    {
+        FScopeLock _(&CS);    
+        TTuple<int32, int32> Handle{IgnoreCategories.Num(), IgnoreCategories.Num() + NewCategories.Num()};
+        IgnoreCategories.Append(NewCategories);
+        return Handle;
+    }
+
+    void RemoveIgnoreCategories(TTuple<int32, int32> Range) 
+    {
+        FScopeLock _(&CS);    
+        for (int32 i=Range.Key; i < Range.Value; ++i)
+        {
+            IgnoreCategories[i] = NAME_None;
+        }
+        if (Algo::AllOf(IgnoreCategories, UE_PROJECTION_MEMBER(FName, IsNone)))
+        {
+            IgnoreCategories.Reset();
+        }
+    }
+
+    TTuple<int32, int32> AddIgnorePatterns(TConstArrayView<FWildcardString> NewPatterns)
+    {
+        FScopeLock _(&CS);
+        TTuple<int32, int32> Handle{IgnorePatterns.Num(), IgnorePatterns.Num() + NewPatterns.Num()};
+        IgnorePatterns.Append(NewPatterns);
+        return Handle;
+    }
+
+    void RemoveIgnorePatterns(TTuple<int32, int32> Range)
+    {
+        FScopeLock _(&CS);    
+        for (int32 i=Range.Key; i < Range.Value; ++i)
+        {
+            IgnorePatterns[i] = FWildcardString{};
+        }
+        if (Algo::AllOf(IgnorePatterns, UE_PROJECTION_MEMBER(FWildcardString, IsEmpty)))
+        {
+            IgnorePatterns.Reset();
+        }
+    }
+
+    virtual bool CanBeUsedOnMultipleThreads() const override
+    {
+        return true;
+    }
+
+    virtual void Serialize(const TCHAR* V, ELogVerbosity::Type Verbosity, const class FName& Category) override
+    {
+        if (Verbosity > ELogVerbosity::Warning)
+        {
+            return;
+        }
+
+        FScopeLock _(&CS);
+
+        if (!Category.IsNone() && IgnoreCategories.Contains(Category))
+        {
+            return;
+        }
+        if (IgnorePatterns.ContainsByPredicate([&V](const FWildcardString& S) { return S.IsEmpty() || S.IsMatch(V); }))
+        {
+            return;
+        }
+        
+        if (Verbosity == ELogVerbosity::Warning)
+        {
+            Warnings.Emplace(V);
+        }
+        else
+        {
+            Errors.Emplace(V);
+        }
+    }
+    
+    void Stop(TArray<FString>& OutWarnings, TArray<FString>& OutErrors)
+    {
+        FScopeLock _(&CS);
+        OutWarnings = MoveTemp(Warnings);
+        OutErrors = MoveTemp(Errors);
+    }
+
+    FCriticalSection CS;
+    
+    TArray<FName> DefaultIgnoreCategories;
+    TArray<FWildcardString> DefaultIgnorePatterns;
+    
+    TArray<FName> IgnoreCategories;
+    TArray<FWildcardString> IgnorePatterns;
+    
+    TArray<FString> Errors;
+    TArray<FString> Warnings;
+};
+
+thread_local FScopedLogMessageGatherer* CurrentGatherer = nullptr;
+
+FScopedLogMessageGatherer::FScopedLogMessageGatherer(bool bEnabled)
+{
+    if (bEnabled)
+    {
+        Impl = new FLogMessageGathererImpl;
+        Previous = CurrentGatherer;
+        CurrentGatherer = this;
+    }
+}
+
+FScopedLogMessageGatherer::~FScopedLogMessageGatherer()
+{
+    if (Impl)
+    {
+        CurrentGatherer = Previous;
+        delete Impl;
+    }
+}
+
+void FScopedLogMessageGatherer::Stop(TArray<FString>& OutWarnings, TArray<FString>& OutErrors)
+{
+    if (Impl)
+    {
+        Impl->Stop(OutWarnings, OutErrors);
+    }
+}
+
+FScopedLogMessageGatherer* FScopedLogMessageGatherer::GetCurrentThreadGatherer()
+{
+    return CurrentGatherer;
+}
+
+TTuple<int32, int32> FScopedLogMessageGatherer::AddIgnoreCategories(TConstArrayView<FName> NewCategories)
+{
+    return Impl->AddIgnoreCategories(NewCategories);
+}
+
+void FScopedLogMessageGatherer::RemoveIgnoreCategories(TTuple<int32, int32> Range)
+{
+    Impl->RemoveIgnorePatterns(Range);
+}
+
+TTuple<int32, int32> FScopedLogMessageGatherer::AddIgnorePatterns(TConstArrayView<FWildcardString> NewPatterns)
+{
+    return Impl->AddIgnorePatterns(NewPatterns);
+}
+
+void FScopedLogMessageGatherer::RemoveIgnorePatterns(TTuple<int32, int32> Range)
+{
+    Impl->RemoveIgnorePatterns(Range);
+}
+
+TDelegate<TArray<FName>(void)>& FScopedLogMessageGatherer::GetDefaultIgnoreCategoriesDelegate()
+{
+    static TDelegate<TArray<FName>(void)> Delegate;
+    return Delegate;
+}
+
+TDelegate<TArray<FWildcardString>(void)>& FScopedLogMessageGatherer::GetDefaultIgnorePatternsDelegate()
+{
+    static TDelegate<TArray<FWildcardString>(void)> Delegate;
+    return Delegate;
+}
+
+FScopedIgnoreLogMessages::FScopedIgnoreLogMessages(TConstArrayView<FWildcardString> Patterns, TConstArrayView<FName> Categories)
+{
+    if(FScopedLogMessageGatherer* Gatherer = FScopedLogMessageGatherer::GetCurrentThreadGatherer())
+    {
+        CategoriesRange = Gatherer->AddIgnoreCategories(Categories);
+        PatternsRange = Gatherer->AddIgnorePatterns(Patterns);
+    }
+}
+
+FScopedIgnoreLogMessages::~FScopedIgnoreLogMessages()
+{
+    if(FScopedLogMessageGatherer* Gatherer = FScopedLogMessageGatherer::GetCurrentThreadGatherer())
+    {
+        Gatherer->RemoveIgnoreCategories(CategoriesRange);
+        Gatherer->RemoveIgnoreCategories(PatternsRange);
+    }
+}
+
+} // namespace UE::DataValidation
\ No newline at end of file
diff --git a/Engine/Plugins/Editor/DataValidation/Source/DataValidation/Private/EditorValidatorSubsystem.cpp b/Engine/Plugins/Editor/DataValidation/Source/DataValidation/Private/EditorValidatorSubsystem.cpp
--- a/Engine/Plugins/Editor/DataValidation/Source/DataValidation/Private/EditorValidatorSubsystem.cpp
+++ b/Engine/Plugins/Editor/DataValidation/Source/DataValidation/Private/EditorValidatorSubsystem.cpp
@@ -2,10 +2,15 @@
 
 #include "EditorValidatorSubsystem.h"
 
+#include "Algo/AnyOf.h"
+#include "Algo/RemoveIf.h"
+#include "AssetCompilingManager.h"
 #include "AssetRegistry/ARFilter.h"
+#include "AssetRegistry/AssetDataToken.h"
 #include "AssetRegistry/AssetRegistryModule.h"
 #include "Blueprint/BlueprintSupport.h"
 #include "DataValidationChangelist.h"
+#include "EditorValidatorHelpers.h"
 #include "DirectoryWatcherModule.h"
 #include "Editor.h"
 #include "EditorUtilityBlueprint.h"
@@ -17,8 +22,10 @@
 #include "Logging/MessageLog.h"
 #include "Misc/App.h"
 #include "Misc/DataValidation.h"
+#include "Misc/MessageDialog.h"
 #include "Misc/PackageName.h"
 #include "Misc/ScopedSlowTask.h"
+#include "Misc/ScopeExit.h"
 #include "UObject/TopLevelAssetPath.h"
 
 #include UE_INLINE_GENERATED_CPP_BY_NAME(EditorValidatorSubsystem)
@@ -27,9 +34,17 @@
 
 DEFINE_LOG_CATEGORY(LogContentValidation);
 
+FValidateAssetsSettings::FValidateAssetsSettings()
+	: ShowMessageLogSeverity(EMessageSeverity::Warning)
+	, MessageLogPageTitle(LOCTEXT("DataValidation.MessagePageTitle", "Data Validation"))
+{
+}
+
+FValidateAssetsSettings::~FValidateAssetsSettings()
+{
+}
 
 UDataValidationSettings::UDataValidationSettings()
-	: bValidateOnSave(true)
 {
 
 }
@@ -40,6 +55,24 @@ UEditorValidatorSubsystem::UEditorValidatorSubsystem()
 	bAllowBlueprintValidators = true;
 }
 
+bool UEditorValidatorSubsystem::ShouldCreateSubsystem(UObject* Outer) const
+{
+	if (GetClass() == UEditorValidatorSubsystem::StaticClass())
+	{
+		TArray<UClass*> ChildClasses;
+		GetDerivedClasses(UEditorValidatorSubsystem::StaticClass(), ChildClasses, true);
+		for (UClass* Child : ChildClasses)
+		{
+			if (Child->GetDefaultObject<UEditorSubsystem>()->ShouldCreateSubsystem(Outer))
+			{
+				// Do not create this class because one of our child classes wants to be created
+				return false;
+			}
+		}
+	}
+	return true;
+}
+
 void UEditorValidatorSubsystem::Initialize(FSubsystemCollectionBase& Collection)
 {
 	const FAssetRegistryModule& AssetRegistryModule = FModuleManager::LoadModuleChecked<FAssetRegistryModule>(TEXT("AssetRegistry"));
@@ -80,6 +113,19 @@ void UEditorValidatorSubsystem::Initialize(FSubsystemCollectionBase& Collection)
 	ISourceControlModule::Get().RegisterPreSubmitDataValidation(FSourceControlPreSubmitDataValidationDelegate::CreateUObject(this, &UEditorValidatorSubsystem::ValidateChangelistPreSubmit));
 }
 
+bool UEditorValidatorSubsystem::ShouldValidateAsset(
+	const FAssetData& 				Asset,
+	const FValidateAssetsSettings& 	Settings,
+	FDataValidationContext& 		InContext) const
+{
+	if (IsPathExcludedFromValidation(WriteToString<FName::StringBufferSize>(Asset.PackageName).ToView()))
+	{
+		return false;
+	}
+
+	return true;
+}
+
 // Rename to BP validators
 void UEditorValidatorSubsystem::RegisterBlueprintValidators()
 {
@@ -158,55 +204,74 @@ void UEditorValidatorSubsystem::CleanupValidators()
 	Validators.Empty();
 }
 
-EDataValidationResult UEditorValidatorSubsystem::IsObjectValid(UObject* InObject, TArray<FText>& ValidationErrors, TArray<FText>& ValidationWarnings, const EDataValidationUsecase InValidationUsecase) const
+EDataValidationResult UEditorValidatorSubsystem::IsObjectValid(
+	UObject* InObject,
+	TArray<FText>& ValidationErrors,
+	TArray<FText>& ValidationWarnings,
+	const EDataValidationUsecase InValidationUsecase) const
+{
+	FDataValidationContext Context(false, InValidationUsecase, {}); // No associated objects in this context
+	EDataValidationResult Result = IsObjectValidWithContext(InObject, Context);
+	Context.SplitIssues(ValidationErrors, ValidationWarnings);
+	return Result;
+}
+
+EDataValidationResult UEditorValidatorSubsystem::IsAssetValid(
+	const FAssetData& AssetData,
+	TArray<FText>& ValidationErrors,
+	TArray<FText>& ValidationWarnings,
+	const EDataValidationUsecase InValidationUsecase) const
+{
+	if (AssetData.IsValid())
+	{
+		UObject* Obj = AssetData.GetAsset({ ULevel::LoadAllExternalObjectsTag });
+		if (Obj)
+		{
+			FDataValidationContext Context(false, InValidationUsecase, {}); // No associated objects in this context
+			EDataValidationResult Result = ValidateObjectInternal(AssetData, Obj, Context);
+			Context.SplitIssues(ValidationErrors, ValidationWarnings);
+			return Result;
+		}
+		return EDataValidationResult::NotValidated;
+	}
+
+	return EDataValidationResult::Invalid;
+}
+
+EDataValidationResult UEditorValidatorSubsystem::IsObjectValidWithContext(
+	UObject* InObject,
+	FDataValidationContext& InContext) const
 {
 	EDataValidationResult Result = EDataValidationResult::NotValidated;
 	
 	if (ensure(InObject))
 	{
-		// First check the class level validation
-		FDataValidationContext Context;
-		Result = const_cast<const UObject*>(InObject)->IsDataValid(Context);
-		Context.SplitIssues(ValidationWarnings, ValidationErrors);
+		FAssetRegistryModule& AssetRegistryModule = FModuleManager::LoadModuleChecked<FAssetRegistryModule>("AssetRegistry");
+		IAssetRegistry& AssetRegistry = AssetRegistryModule.Get();
 
-		// If the asset is still valid or there wasn't a class-level validation, keep validating with custom validators
-		if (Result != EDataValidationResult::Invalid)
+		FAssetData AssetData = AssetRegistry.GetAssetByObjectPath(FSoftObjectPath(InObject), true);
+		if (!AssetData.IsValid())
 		{
-			// This doesn't change the logical view of the subsystem
-			const_cast<UEditorValidatorSubsystem&>(*this).LoadValidators();
-
-			for (const TPair<FTopLevelAssetPath, TObjectPtr<UEditorValidatorBase>>& ValidatorPair : Validators)
-			{
-				if (ValidatorPair.Value && ValidatorPair.Value->IsEnabled() && ValidatorPair.Value->CanValidate(InValidationUsecase) && ValidatorPair.Value->CanValidateAsset(InObject))
-				{
-					ValidatorPair.Value->ResetValidationState();
-					EDataValidationResult NewResult = ValidatorPair.Value->ValidateLoadedAsset(InObject, ValidationErrors);
-
-					// Don't accidentally overwrite an invalid result with a valid or not-validated one
-					if(Result != EDataValidationResult::Invalid)
-					{
-						Result = NewResult;
-					}
-
-					ValidationWarnings.Append(ValidatorPair.Value->GetAllWarnings());
-
-					ensureMsgf(ValidatorPair.Value->IsValidationStateSet(), TEXT("Validator %s did not include a pass or fail state."), *ValidatorPair.Value->GetClass()->GetName());
-				}
-			}
+			// Construct dynamically with potentially fewer tags 
+			AssetData = FAssetData(InObject);
 		}
+		
+		return ValidateObjectInternal(AssetData, InObject, InContext);
 	}
 
 	return Result;
 }
 
-EDataValidationResult UEditorValidatorSubsystem::IsAssetValid(const FAssetData& AssetData, TArray<FText>& ValidationErrors, TArray<FText>& ValidationWarnings, const EDataValidationUsecase InValidationUsecase) const
+EDataValidationResult UEditorValidatorSubsystem::IsAssetValidWithContext(
+	const FAssetData& AssetData,
+	FDataValidationContext& InContext) const
 {
 	if (AssetData.IsValid())
 	{
 		UObject* Obj = AssetData.GetAsset({ ULevel::LoadAllExternalObjectsTag });
 		if (Obj)
 		{
-			return IsObjectValid(Obj, ValidationErrors, ValidationWarnings, InValidationUsecase);
+			return ValidateObjectInternal(AssetData, Obj, InContext);
 		}
 		return EDataValidationResult::NotValidated;
 	}
@@ -214,139 +279,281 @@ EDataValidationResult UEditorValidatorSubsystem::IsAssetValid(const FAssetData&
 	return EDataValidationResult::Invalid;
 }
 
-int32 UEditorValidatorSubsystem::ValidateAssets(TArray<FAssetData> AssetDataList, bool bSkipExcludedDirectories, bool bShowIfNoFailures) const
+EDataValidationResult UEditorValidatorSubsystem::ValidateObjectInternal(
+	const FAssetData& InAssetData,
+	UObject* InObject,
+	FDataValidationContext& InContext) const
 {
-	FValidateAssetsSettings Settings;
-	FValidateAssetsResults Results;
+	EDataValidationResult Result = EDataValidationResult::NotValidated;
+	
+	if (ensure(InObject) && ensure(InAssetData.IsValid()))
+	{
+		// First check the class level validation
+		Result = const_cast<const UObject*>(InObject)->IsDataValid(InContext);
+
+		// If the asset is still valid or there wasn't a class-level validation, keep validating with custom validators
+		if (Result == EDataValidationResult::Invalid)
+		{
+			return Result;
+		}
+			
+		// This doesn't change the logical view of the subsystem
+		const_cast<UEditorValidatorSubsystem&>(*this).LoadValidators();
+
+		for (const TPair<FTopLevelAssetPath, TObjectPtr<UEditorValidatorBase>>& ValidatorPair : Validators)
+		{
+			UEditorValidatorBase* Validator = ValidatorPair.Value;
+			if (Validator && Validator->IsEnabled())
+			{
+				EDataValidationResult NewResult = Validator->ValidateLoadedAsset(InAssetData, InObject, InContext);
+				Result = CombineDataValidationResults(Result, NewResult);
+			}
+		}
+	}
 
-	Settings.bSkipExcludedDirectories = bSkipExcludedDirectories;
-	Settings.bShowIfNoFailures = bShowIfNoFailures;
+	return Result;
+}
+
+int32 UEditorValidatorSubsystem::ValidateAssetsWithSettings(
+	const TArray<FAssetData>& AssetDataList,
+	FValidateAssetsSettings& InSettings,
+	FValidateAssetsResults& OutResults) const
+{
+	FMessageLog DataValidationLog(UE::DataValidation::MessageLogName);
+	ValidateAssetsInternal(DataValidationLog, TSet<FAssetData>{AssetDataList}, InSettings, OutResults);
 	
-	return ValidateAssetsWithSettings(AssetDataList, Settings, Results);
+	if (EMessageSeverity::Type* Severity = InSettings.ShowMessageLogSeverity.GetPtrOrNull())
+	{
+		DataValidationLog.Open(*Severity, false);
+	}
+
+	return OutResults.NumWarnings + OutResults.NumInvalid;
 }
 
-int32 UEditorValidatorSubsystem::ValidateAssetsWithSettings(const TArray<FAssetData>& AssetDataList, const FValidateAssetsSettings& InSettings, FValidateAssetsResults& OutResults) const
+void UEditorValidatorSubsystem::ValidateAssetsInternal(
+	FMessageLog& 					DataValidationLog,
+	TSet<FAssetData>				AssetDataList,
+	const FValidateAssetsSettings& 	InSettings,
+	FValidateAssetsResults& 		OutResults) const
 {
-	FScopedSlowTask SlowTask(1.0f, LOCTEXT("ValidatingDataTask", "Validating Data..."));
-	SlowTask.Visibility = ESlowTaskVisibility::ForceVisible;
-	SlowTask.MakeDialogDelayed(.1f);
+	FAssetRegistryModule& AssetRegistryModule = FModuleManager::LoadModuleChecked<FAssetRegistryModule>("AssetRegistry");
+	IAssetRegistry& AssetRegistry = AssetRegistryModule.Get();
 
+	FScopedSlowTask SlowTask(AssetDataList.Num(), LOCTEXT("DataValidation.ValidateAssetsTask", "Validating Assets"));
+	SlowTask.MakeDialog();
+	
 	// Broadcast the Editor event before we start validating. This lets other systems (such as Sequencer) restore the state
 	// of the level to what is actually saved on disk before performing validation.
 	if (FEditorDelegates::OnPreAssetValidation.IsBound())
 	{
 		FEditorDelegates::OnPreAssetValidation.Broadcast();
 	}
-
-	FMessageLog DataValidationLog("AssetCheck");
-
-	int32 NumAdded = 0;
-
-	int32 NumFilesChecked = 0;
-	int32 NumValidFiles = 0;
-	int32 NumInvalidFiles = 0;
-	int32 NumFilesSkipped = 0;
-	int32 NumFilesUnableToValidate = 0;
-	int32 NumFilesWithWarnings = 0;
-
-	int32 NumFilesToValidate = AssetDataList.Num();
-
-	auto AddAssetLogTokens = [&DataValidationLog](EMessageSeverity::Type Severity, const TArray<FText>& TextMessages, FName PackageName)
+	
+	// Filter external objects out from the asset data list to be validated indirectly via their outers 
+	TMap<FAssetData, TArray<FAssetData>> AssetsToExternalObjects;
+	for (auto It = AssetDataList.CreateIterator(); It; ++It)
 	{
-		const FString PackageNameString = PackageName.ToString();
-		for (const FText& Msg : TextMessages)
+		if (!It->GetOptionalOuterPathName().IsNone())
 		{
-			const FString AssetLogString = FAssetMsg::GetAssetLogString(*PackageNameString, Msg.ToString());
-			FString BeforeAsset;
-			FString AfterAsset;
-			TSharedRef<FTokenizedMessage> TokenizedMessage = DataValidationLog.Message(Severity);
-			if (AssetLogString.Split(PackageNameString, &BeforeAsset, &AfterAsset))
-			{
-				if (!BeforeAsset.IsEmpty())
-				{
-					TokenizedMessage->AddToken(FTextToken::Create(FText::FromString(BeforeAsset)));
-				}
-				TokenizedMessage->AddToken(FAssetNameToken::Create(PackageNameString));
-				if (!AfterAsset.IsEmpty())
-				{
-					TokenizedMessage->AddToken(FTextToken::Create(FText::FromString(AfterAsset)));
-				}
-			}
-			else
+			FSoftObjectPath Path = It->ToSoftObjectPath();
+			FAssetData OuterAsset = AssetRegistry.GetAssetByObjectPath(Path.GetWithoutSubPath(), true);
+			AssetsToExternalObjects.FindOrAdd(OuterAsset).Add(*It);
+			It.RemoveCurrent();
+		}
+	}
+	
+	// Add any packages which contain those external objects to be validated
+	{
+		FDataValidationContext ValidationContext(false, InSettings.ValidationUsecase, {});
+		for (const TPair<FAssetData, TArray<FAssetData>>& Pair : AssetsToExternalObjects)
+		{
+			if (ShouldValidateAsset(Pair.Key, InSettings, ValidationContext))
 			{
-				TokenizedMessage->AddToken(FTextToken::Create(FText::FromString(AssetLogString)));
+				AssetDataList.Add(Pair.Key);
 			}
 		}
-	};
+		UE::DataValidation::AddAssetValidationMessages(DataValidationLog, ValidationContext);
+		DataValidationLog.Flush();
+	}
+
+	// Dont let other async compilation warnings be attributed incorrectly to the package that is loading.
+	if (FAssetCompilingManager::Get().GetNumRemainingAssets())
+	{
+		FScopedSlowTask CompileAssetsSlowTask(0.f, LOCTEXT("DataValidation.CompilingAssetsBeforeCheckingContentTask", "Finishing asset compilations before checking content..."));
+		CompileAssetsSlowTask.MakeDialog();
+		FAssetCompilingManager::Get().FinishAllCompilation();
+	}
+
+	OutResults.NumRequested = AssetDataList.Num();
 
 	// Now add to map or update as needed
 	for (const FAssetData& Data : AssetDataList)
 	{
-		FText ValidatingMessage = FText::Format(LOCTEXT("ValidatingFilename", "Validating {0}"), FText::FromString(Data.GetFullName()));
-		SlowTask.EnterProgressFrame(1.0f / NumFilesToValidate, ValidatingMessage);
+		ensure(Data.IsValid());
+
+		SlowTask.EnterProgressFrame(1.0f, FText::Format(LOCTEXT("DataValidation.ValidatingFilename", "Validating {0}"), FText::FromString(Data.GetFullName())));
+		
+		if (OutResults.NumChecked >= InSettings.MaxAssetsToValidate)
+		{
+			OutResults.bAssetLimitReached = true;
+			DataValidationLog.Info(FText::Format(LOCTEXT("DataValidation.MaxAssetCountReached", "MaxAssetsToValidate count {0} reached."), InSettings.MaxAssetsToValidate));
+			break;
+		}
 
 		if (Data.HasAnyPackageFlags(PKG_Cooked))
 		{
-			++NumFilesSkipped;
+			++OutResults.NumSkipped;
 			continue;
 		}
 
 		// Check exclusion path
 		if (InSettings.bSkipExcludedDirectories && IsPathExcludedFromValidation(Data.PackageName.ToString()))
 		{
-			++NumFilesSkipped;
+			++OutResults.NumSkipped;
 			continue;
 		}
 
 		const bool bLoadAsset = false;
 		if (!InSettings.bLoadAssetsForValidation && !Data.FastGetAsset(bLoadAsset))
 		{
-			++NumFilesSkipped;
+			++OutResults.NumSkipped;
 			continue;
 		}
-
-		UE_LOG(LogContentValidation, Display, TEXT("%s"), *ValidatingMessage.ToString());
-
-		TArray<FText> ValidationErrors;
-		TArray<FText> ValidationWarnings;
-		EDataValidationResult Result = IsAssetValid(Data, ValidationErrors, ValidationWarnings, InSettings.ValidationUsecase);
-		++NumFilesChecked;
-
-		AddAssetLogTokens(EMessageSeverity::Error, ValidationErrors, Data.PackageName);
-
-		if (ValidationWarnings.Num() > 0)
+		
+		DataValidationLog.Info()
+			->AddToken(FAssetDataToken::Create(Data))
+			->AddToken(FTextToken::Create(LOCTEXT("Data.ValidatingAsset", "Validating asset")));
+		UE_LOG(LogContentValidation, Log, TEXT("Validating asset %s"), *Data.ToSoftObjectPath().ToString());
+		
+		UObject* LoadedAsset = Data.FastGetAsset(false);
+		const bool bAlreadyLoaded = LoadedAsset != nullptr;
+
+		FDataValidationContext ValidationContext(bAlreadyLoaded, InSettings.ValidationUsecase, AssetsToExternalObjects.FindRef(Data));
+		EDataValidationResult Result = EDataValidationResult::NotValidated;
+		if (!LoadedAsset)
 		{
-			++NumFilesWithWarnings;
-
-			AddAssetLogTokens(EMessageSeverity::Warning, ValidationWarnings, Data.PackageName);
+			UE_LOG(LogContentValidation, Log, TEXT("Loading asset %s for validation"), *Data.ToSoftObjectPath().ToString());
+			UE::DataValidation::FScopedLogMessageGatherer LogGatherer(InSettings.bCaptureAssetLoadLogs);
+			LoadedAsset = Data.GetAsset(); // Do not load external objects, validators should load external objects that they want via GetAssociatedExternalObjects in the validation context 
+			
+			if (FAssetCompilingManager::Get().GetNumRemainingAssets())
+			{
+				FScopedSlowTask CompileAssetsSlowTask(0.f, LOCTEXT("DataValidation.CompilingAssetsBeforeCheckingContentTask", "Finishing asset compilations before checking content..."));
+				CompileAssetsSlowTask.MakeDialog();
+				FAssetCompilingManager::Get().FinishAllCompilation();
+			}
+			
+			// Associate any load errors with this asset in the message log
+			TArray<FString> Warnings;
+			TArray<FString> Errors;
+			LogGatherer.Stop(Warnings, Errors);
+			if (Warnings.Num())
+			{
+				TStringBuilder<2048> Buffer;
+				Buffer.Join(Warnings, LINE_TERMINATOR);
+				ValidationContext.AddMessage(EMessageSeverity::Warning)
+					->AddToken(FAssetDataToken::Create(Data))
+					->AddText(LOCTEXT("DataValidation.LoadWarnings", "Warnings loading asset {0}"), FText::FromStringView(Buffer.ToView()));
+			}
+			if(Errors.Num())
+			{
+				TStringBuilder<2048> Buffer;
+				Buffer.Join(Errors, LINE_TERMINATOR);
+				ValidationContext.AddMessage(EMessageSeverity::Error)
+					->AddToken(FAssetDataToken::Create(Data))
+					->AddText(LOCTEXT("DataValidation.LoadErrors", "Errors loading asset {0}"), FText::FromStringView(Buffer.ToView()));
+				Result = EDataValidationResult::Invalid;
+			}
 		}
 
-		if (Result == EDataValidationResult::Valid)
+		if (LoadedAsset)
 		{
-			if (ValidationWarnings.Num() > 0)
+			UE::DataValidation::FScopedLogMessageGatherer LogGatherer(InSettings.bCaptureLogsDuringValidation);
+			Result = IsObjectValidWithContext(LoadedAsset, ValidationContext);
+			
+			// Associate any UE_LOG errors with this asset in the message log
+			TArray<FString> Warnings;
+			TArray<FString> Errors;
+			LogGatherer.Stop(Warnings, Errors);
+			if (Warnings.Num())
 			{
-				DataValidationLog.Info()->AddToken(FAssetNameToken::Create(Data.PackageName.ToString()))
-					->AddToken(FTextToken::Create(LOCTEXT("ContainsWarningsResult", "contains valid data, but has warnings.")));
+				TStringBuilder<2048> Buffer;
+				Buffer.Join(Warnings, LINE_TERMINATOR);
+				ValidationContext.AddMessage(EMessageSeverity::Error)
+					->AddToken(FAssetDataToken::Create(Data))
+					->AddText(LOCTEXT("DataValidation.DuringValidationWarnings", "Warnings logged while validating asset {0}"), FText::FromStringView(Buffer.ToView()));
+				Result = EDataValidationResult::Invalid;
 			}
-			++NumValidFiles;
-		}
-		else
-		{
-			if (Result == EDataValidationResult::Invalid)
+			if(Errors.Num())
 			{
-				DataValidationLog.Info()->AddToken(FAssetNameToken::Create(Data.PackageName.ToString()))
-					->AddToken(FTextToken::Create(LOCTEXT("InvalidDataResult", "contains invalid data.")));
-				++NumInvalidFiles;
+				TStringBuilder<2048> Buffer;
+				Buffer.Join(Errors, LINE_TERMINATOR);
+				ValidationContext.AddMessage(EMessageSeverity::Error)
+					->AddToken(FAssetDataToken::Create(Data))
+					->AddText(LOCTEXT("DataValidation.DuringValidationErrors", "Errors logged while validating asset {0}"), FText::FromStringView(Buffer.ToView()));
+				Result = EDataValidationResult::Invalid;
 			}
-			else if (Result == EDataValidationResult::NotValidated)
+		}
+		else if (InSettings.bLoadAssetsForValidation)
+		{
+			ValidationContext.AddMessage(EMessageSeverity::Error)
+				->AddToken(FAssetDataToken::Create(Data))
+				->AddToken(FTextToken::Create(LOCTEXT("DataValidation.LoadFailed", "Failed to load object")));
+			Result = EDataValidationResult::Invalid;
+		}
+		else 
+		{
+			ValidationContext.AddMessage(EMessageSeverity::Error)
+				->AddToken(FAssetDataToken::Create(Data))
+				->AddToken(FTextToken::Create(LOCTEXT("DataValidation.CannotValidateNotLoaded", "Cannot validate unloaded asset")));
+			Result = EDataValidationResult::Invalid;
+		}
+
+		++OutResults.NumChecked;
+		
+		// Don't add more messages to ValidationContext after this point because we will no longer add them to the message log 
+		UE::DataValidation::AddAssetValidationMessages(Data, DataValidationLog, ValidationContext);
+
+		bool bAnyWarnings = Algo::AnyOf(ValidationContext.GetIssues(), [](const FDataValidationContext::FIssue& Issue){ return Issue.Severity == EMessageSeverity::Warning; });
+		if (bAnyWarnings)
+		{
+			++OutResults.NumWarnings;
+		}
+
+		if (InSettings.bShowIfNoFailures)
+		{
+			switch (Result)
 			{
-				if (InSettings.bShowIfNoFailures)
-				{
-					DataValidationLog.Info()->AddToken(FAssetNameToken::Create(Data.PackageName.ToString()))
-						->AddToken(FTextToken::Create(LOCTEXT("NotValidatedDataResult", "has no data validation.")));
-				}
-				++NumFilesUnableToValidate;
+				case EDataValidationResult::Valid:
+					if (bAnyWarnings)
+					{
+						DataValidationLog.Info()->AddToken(FAssetDataToken::Create(Data))
+							->AddToken(FTextToken::Create(LOCTEXT("DataValidation.ContainsWarningsResult", "contains valid data, but has warnings.")));
+					}
+					else 
+					{
+						DataValidationLog.Info()->AddToken(FAssetDataToken::Create(Data))
+							->AddToken(FTextToken::Create(LOCTEXT("DataValidation.ContainsWarningsResult", "contains valid data.")));
+					}
+					break;
+				case EDataValidationResult::NotValidated:
+					DataValidationLog.Info()->AddToken(FAssetDataToken::Create(Data))
+						->AddToken(FTextToken::Create(LOCTEXT("DataValidation.NotValidatedDataResult", "has no data validation.")));
+					break;
 			}
+		}	
+
+		switch (Result)
+		{
+			case EDataValidationResult::Valid:
+				++OutResults.NumValid;
+				break;
+			case EDataValidationResult::Invalid:
+				++OutResults.NumInvalid;
+				break;
+			case EDataValidationResult::NotValidated:
+				++OutResults.NumUnableToValidate;
+				break;
 		}
 
 		if (InSettings.bCollectPerAssetDetails)
@@ -355,49 +562,39 @@ int32 UEditorValidatorSubsystem::ValidateAssetsWithSettings(const TArray<FAssetD
 			Details.PackageName = Data.PackageName;
 			Details.AssetName = Data.AssetName;
 			Details.Result = Result;
-			Details.ValidationErrors = MoveTemp(ValidationErrors);
-			Details.ValidationWarnings = MoveTemp(ValidationWarnings);
+			ValidationContext.SplitIssues(Details.ValidationErrors, Details.ValidationWarnings);
 		}
+		
+		DataValidationLog.Flush();
 	}
 
-	const bool bFailed = (NumInvalidFiles > 0);
-	const bool bAtLeastOneWarning = (NumFilesWithWarnings > 0);
+	// Broadcast now that we're complete so other systems can go back to their previous state.
+	if (FEditorDelegates::OnPostAssetValidation.IsBound())
+	{
+		FEditorDelegates::OnPostAssetValidation.Broadcast();
+	}
+}
 
-	OutResults.NumChecked = NumFilesChecked;
-	OutResults.NumValid = NumValidFiles;
-	OutResults.NumInvalid = NumInvalidFiles;
-	OutResults.NumSkipped = NumFilesSkipped;
-	OutResults.NumWarnings = NumFilesWithWarnings;
-	OutResults.NumUnableToValidate = NumFilesUnableToValidate;
+void UEditorValidatorSubsystem::LogAssetValidationSummary(FMessageLog& DataValidationLog, const FValidateAssetsSettings& InSettings, const FValidateAssetsResults& Results) const
+{
+	const bool bFailed = (Results.NumInvalid > 0);
+	const bool bAtLeastOneWarning = (Results.NumWarnings > 0);
 
 	if (bFailed || bAtLeastOneWarning || InSettings.bShowIfNoFailures)
 	{
 		FFormatNamedArguments Arguments;
 		Arguments.Add(TEXT("Result"), bFailed ? LOCTEXT("Failed", "FAILED") : LOCTEXT("Succeeded", "SUCCEEDED"));
-		Arguments.Add(TEXT("NumChecked"), NumFilesChecked);
-		Arguments.Add(TEXT("NumValid"), NumValidFiles);
-		Arguments.Add(TEXT("NumInvalid"), NumInvalidFiles);
-		Arguments.Add(TEXT("NumSkipped"), NumFilesSkipped);
-		Arguments.Add(TEXT("NumUnableToValidate"), NumFilesUnableToValidate);
+		Arguments.Add(TEXT("NumChecked"), Results.NumChecked);
+		Arguments.Add(TEXT("NumValid"), Results.NumValid);
+		Arguments.Add(TEXT("NumInvalid"), Results.NumInvalid);
+		Arguments.Add(TEXT("NumSkipped"), Results.NumSkipped);
+		Arguments.Add(TEXT("NumUnableToValidate"), Results.NumUnableToValidate);
 
-		DataValidationLog.Info()->AddToken(FTextToken::Create(FText::Format(LOCTEXT("SuccessOrFailure", "Data validation {Result}."), Arguments)));
-		DataValidationLog.Info()->AddToken(FTextToken::Create(FText::Format(LOCTEXT("ResultsSummary", "Files Checked: {NumChecked}, Passed: {NumValid}, Failed: {NumInvalid}, Skipped: {NumSkipped}, Unable to validate: {NumUnableToValidate}"), Arguments)));
+		DataValidationLog.Info()->AddToken(FTextToken::Create(FText::Format(LOCTEXT("DataValidation.SuccessOrFailure", "Data validation {Result}."), Arguments)));
+		DataValidationLog.Info()->AddToken(FTextToken::Create(FText::Format(LOCTEXT("DataValidation.ResultsSummary", "Files Checked: {NumChecked}, Passed: {NumValid}, Failed: {NumInvalid}, Skipped: {NumSkipped}, Unable to validate: {NumUnableToValidate}"), Arguments)));
 
 		DataValidationLog.Open(EMessageSeverity::Info, true);
 	}
-
-	// Broadcast now that we're complete so other systems can go back to their previous state.
-	if (FEditorDelegates::OnPostAssetValidation.IsBound())
-	{
-		FEditorDelegates::OnPostAssetValidation.Broadcast();
-	}
-
-	return NumInvalidFiles + NumFilesWithWarnings;
-}
-
-void UEditorValidatorSubsystem::ValidateOnSave(TArray<FAssetData> AssetDataList) const
-{
-	ValidateOnSave(AssetDataList, false /* bProceduralSave */);
 }
 
 void UEditorValidatorSubsystem::ValidateOnSave(TArray<FAssetData> AssetDataList, bool bProceduralSave) const
@@ -412,18 +609,26 @@ void UEditorValidatorSubsystem::ValidateOnSave(TArray<FAssetData> AssetDataList,
 	{
 		return;
 	}
+	
+	FValidateAssetsSettings Settings;
+	{
+		FDataValidationContext Context {false, EDataValidationUsecase::Save, {}};
+		AssetDataList.SetNum(Algo::RemoveIf(AssetDataList, [this, &Settings, &Context](const FAssetData& Asset) { 
+			return !ShouldValidateAsset(Asset, Settings, Context); 
+		}));
+	}
 
-	FMessageLog DataValidationLog("AssetCheck");
+	FMessageLog DataValidationLog(UE::DataValidation::MessageLogName);
 	FText SavedAsset = AssetDataList.Num() == 1 ? FText::FromName(AssetDataList[0].AssetName) : LOCTEXT("MultipleErrors", "multiple assets");
 	DataValidationLog.NewPage(FText::Format(LOCTEXT("DataValidationLogPage", "Asset Save: {0}"), SavedAsset));
 
-	FValidateAssetsSettings Settings;
 	FValidateAssetsResults Results;
 
 	Settings.bSkipExcludedDirectories = true;
 	Settings.bShowIfNoFailures = false;
 	Settings.ValidationUsecase = EDataValidationUsecase::Save;
 	Settings.bLoadAssetsForValidation = false;
+	Settings.MessageLogPageTitle = LOCTEXT("DataValidation.ValidateOnSaveMessageLogPage", "Validate Saved Assets");
 
 	if (ValidateAssetsWithSettings(AssetDataList, Settings, Results) > 0)
 	{
@@ -433,11 +638,6 @@ void UEditorValidatorSubsystem::ValidateOnSave(TArray<FAssetData> AssetDataList,
 	}
 }
 
-void UEditorValidatorSubsystem::ValidateSavedPackage(FName PackageName)
-{
-	ValidateSavedPackage(PackageName, false /* bProceduralSave */);
-}
-
 void UEditorValidatorSubsystem::ValidateSavedPackage(FName PackageName, bool bProceduralSave)
 {
 	// Only validate if enabled and not auto saving
@@ -460,7 +660,7 @@ void UEditorValidatorSubsystem::ValidateSavedPackage(FName PackageName, bool bPr
 	SavedPackagesToValidate.AddUnique(PackageName);
 }
 
-bool UEditorValidatorSubsystem::IsPathExcludedFromValidation(const FString& Path) const
+bool UEditorValidatorSubsystem::IsPathExcludedFromValidation(FStringView Path) const
 {
 	for (const FDirectoryPath& ExcludedPath : ExcludedDirectories)
 	{
@@ -502,23 +702,183 @@ void UEditorValidatorSubsystem::ValidateAllSavedPackages()
 	SavedPackagesToValidate.Empty();
 }
 
-void UEditorValidatorSubsystem::ValidateChangelistPreSubmit(FSourceControlChangelistPtr InChangelist, EDataValidationResult& OutResult, TArray<FText>& OutValidationErrors, TArray<FText>& OutValidationWarnings) const
+void UEditorValidatorSubsystem::ValidateChangelistPreSubmit(
+	FSourceControlChangelistPtr InChangelist,
+	EDataValidationResult& 		OutResult,
+	TArray<FText>& 				OutValidationErrors,
+	TArray<FText>& 				OutValidationWarnings) const
 {
-	FScopedSlowTask SlowTask(1.0f, LOCTEXT("ValidatingChangelistTask", "Validating changelist..."));
-	SlowTask.Visibility = ESlowTaskVisibility::Invisible;
-	SlowTask.MakeDialogDelayed(.1f);
-
-	check(InChangelist);
+	check(InChangelist.IsValid());
 
 	// Create temporary changelist object to do most of the heavy lifting
 	UDataValidationChangelist* Changelist = NewObject<UDataValidationChangelist>();
 	Changelist->Initialize(InChangelist);
-	Changelist->AddToRoot();
 
-	SlowTask.EnterProgressFrame(1.0f, LOCTEXT("ValidatingChangelist", "Validating changelist to submit"));
-	OutResult = IsObjectValid(Changelist, OutValidationErrors, OutValidationWarnings, EDataValidationUsecase::PreSubmit);
+	FValidateAssetsSettings Settings;
+	Settings.ValidationUsecase = EDataValidationUsecase::PreSubmit;
+	Settings.bLoadAssetsForValidation = GetDefault<UDataValidationSettings>()->bLoadAssetsWhenValidatingChangelists;
+	Settings.MessageLogPageTitle = FText::Format(LOCTEXT("DataValidation.ValidateChangelistMessageLogPage", "Validate Changelist {0}"), FText::FromString(InChangelist->GetIdentifier()));
+	FValidateAssetsResults Results;
+	OutResult = ValidateChangelist(Changelist, Settings, Results);
+	
+	if (const FValidateAssetsDetails* Details = Results.AssetsDetails.Find(Changelist->GetPathName()))
+	{
+		OutValidationWarnings = Details->ValidationWarnings;
+		OutValidationErrors = Details->ValidationErrors;
+		
+		for (const TSharedRef<FTokenizedMessage>& Message : Details->ValidationMessages)
+		{
+			switch (Message->GetSeverity())
+			{
+				case EMessageSeverity::Error:
+					OutValidationErrors.Add(Message->ToText());
+					break;
+				case EMessageSeverity::Warning:
+				case EMessageSeverity::PerformanceWarning:
+					OutValidationWarnings.Add(Message->ToText());
+					break;
+			}
+		}
+	}
+}
+
+EDataValidationResult UEditorValidatorSubsystem::ValidateChangelist(
+	UDataValidationChangelist* InChangelist,
+	const FValidateAssetsSettings& InSettings,
+	FValidateAssetsResults& OutResults) const
+{
+	return ValidateChangelistsInternal(MakeArrayView(&InChangelist, 1), InSettings, OutResults);	
+}
+
+EDataValidationResult UEditorValidatorSubsystem::ValidateChangelists(
+	const TArray<UDataValidationChangelist*> InChangelists,
+	const FValidateAssetsSettings& InSettings,
+	FValidateAssetsResults& OutResults) const
+{
+	return ValidateChangelistsInternal(MakeArrayView(InChangelists), InSettings, OutResults);	
+}
+
+EDataValidationResult UEditorValidatorSubsystem::ValidateChangelistsInternal(
+	TConstArrayView<UDataValidationChangelist*>	Changelists,
+	const FValidateAssetsSettings& 				Settings,
+	FValidateAssetsResults& 					OutResults) const
+{
+	FScopedSlowTask SlowTask(Changelists.Num(), LOCTEXT("DataValidation.ValidatingChangelistTask", "Validating Changelists"));
+	SlowTask.Visibility = ESlowTaskVisibility::Invisible;
+	SlowTask.MakeDialog();
+
+	FAssetRegistryModule& AssetRegistryModule = FModuleManager::LoadModuleChecked<FAssetRegistryModule>("AssetRegistry");
+	IAssetRegistry& AssetRegistry = AssetRegistryModule.Get();
+
+	if (AssetRegistry.IsLoadingAssets())
+	{
+		FMessageDialog::Open(EAppMsgType::Ok, LOCTEXT("DataValidation.UnableToValidate_PendingAssetRegistry", "Unable to validate changelist while asset registry scan is in progress. Wait until asset discovery is complete."));
+		return EDataValidationResult::NotValidated;
+	}
+
+	// Choose a specific message log page for this output, flushing in case other recursive calls also write to this log 
+	{	
+		FMessageLog DataValidationLog(UE::DataValidation::MessageLogName);
+		DataValidationLog.SetCurrentPage(Settings.MessageLogPageTitle);
+	}
+
+	FMessageLog DataValidationLog(UE::DataValidation::MessageLogName);
+
+	for (UDataValidationChangelist* CL : Changelists)
+	{
+		CL->AddToRoot();
+	}
+	
+	ON_SCOPE_EXIT {
+		for (UDataValidationChangelist* CL : Changelists)
+		{
+			CL->RemoveFromRoot();		
+		}
+	};
+
+	EDataValidationResult Result = EDataValidationResult::NotValidated;
+	for (UDataValidationChangelist* Changelist : Changelists)
+	{
+		FText ValidationMessage = FText::Format(LOCTEXT("DataValidation.ValidatingChangelistMessage", "Validating changelist {0}"), Changelist->Description);
+		DataValidationLog.Info(ValidationMessage);
+		SlowTask.EnterProgressFrame(1.0f, ValidationMessage);
+
+		FValidateAssetsDetails& Details = OutResults.AssetsDetails.FindOrAdd(Changelist->GetPathName());
+		{
+			FDataValidationContext ValidationContext(false, Settings.ValidationUsecase, {}); // No associated objects for changelist
+			Details.Result = IsObjectValidWithContext(Changelist, ValidationContext);
+			UE::DataValidation::AddAssetValidationMessages(DataValidationLog, ValidationContext);
+		}
+		Result = CombineDataValidationResults(Result, Details.Result);
+		DataValidationLog.Flush();
+	}
+
+	TSet<FAssetData> AssetsToValidate;	
+	for (UDataValidationChangelist* Changelist : Changelists)
+	{
+		FDataValidationContext ValidationContext(false, Settings.ValidationUsecase, {}); // No associated objects for changelist
+		GatherAssetsToValidateFromChangelist(Changelist, Settings, AssetsToValidate, ValidationContext);
+		UE::DataValidation::AddAssetValidationMessages(DataValidationLog, ValidationContext);
+		DataValidationLog.Flush();
+	}
+
+	// Filter out assets that we don't want to validate
+	{
+		FDataValidationContext ValidationContext(false, Settings.ValidationUsecase, {}); 
+		for (auto It = AssetsToValidate.CreateIterator(); It; ++It)
+		{
+			if (!ShouldValidateAsset(*It, Settings, ValidationContext))
+			{
+				UE_LOG(LogContentValidation, Log, TEXT("Excluding asset %s from validation"), *It->GetSoftObjectPath().ToString());
+				It.RemoveCurrent();
+			}
+		}
+		UE::DataValidation::AddAssetValidationMessages(DataValidationLog, ValidationContext);
+		DataValidationLog.Flush();
+	}
+
+	// Validate assets from all changelists
+	ValidateAssetsInternal(DataValidationLog, MoveTemp(AssetsToValidate), Settings, OutResults);
+	LogAssetValidationSummary(DataValidationLog, Settings, OutResults);
+
+	return Result;
+}
 
-	Changelist->RemoveFromRoot();
+void UEditorValidatorSubsystem::GatherAssetsToValidateFromChangelist(
+	UDataValidationChangelist* 		InChangelist,
+	const FValidateAssetsSettings& 	Settings,
+	TSet<FAssetData>& 				OutAssets,
+	FDataValidationContext& 		InContext) const
+{
+	FAssetRegistryModule& AssetRegistryModule = FModuleManager::LoadModuleChecked<FAssetRegistryModule>("AssetRegistry");
+	IAssetRegistry& AssetRegistry = AssetRegistryModule.Get();
+
+	for (const FName& PackageName : InChangelist->ModifiedPackageNames)
+	{
+		TArray<FAssetData> NewAssets;
+		AssetRegistry.GetAssetsByPackageName(PackageName, NewAssets, true);	
+		OutAssets.Append(NewAssets);
+	}
+	
+	if (Settings.bValidateReferencersOfDeletedAssets)
+	{
+		TSet<FAssetData> AdditionalAssetsToValidate;
+		for (FName DeletedPackageName : InChangelist->DeletedPackageNames)
+		{
+			InContext.AddMessage(EMessageSeverity::Info, 
+				FText::Format(LOCTEXT("DataValidation.AddDeletedPackageReferencers", "Adding referencers of deleted package {0} to be validated"), FText::FromName(DeletedPackageName)));
+
+			TArray<FName> PackageReferencers;
+			AssetRegistry.GetReferencers(DeletedPackageName, PackageReferencers, UE::AssetRegistry::EDependencyCategory::Package);
+			for (const FName& Referencer : PackageReferencers)
+			{
+				UE_LOG(LogContentValidation, Log, TEXT("Adding %s to to validated as it is a referencer of deleted asset %s"), *Referencer.ToString(), *DeletedPackageName.ToString());
+				TArray<FAssetData> NewAssets;
+				AssetRegistry.GetAssetsByPackageName(Referencer, NewAssets, true);
+				OutAssets.Append(NewAssets);
+			}
+		}
+	}
 }
 
 void UEditorValidatorSubsystem::LoadValidators()
diff --git a/Engine/Plugins/Editor/DataValidation/Source/DataValidation/Private/EditorValidator_Localization.cpp b/Engine/Plugins/Editor/DataValidation/Source/DataValidation/Private/EditorValidator_Localization.cpp
--- a/Engine/Plugins/Editor/DataValidation/Source/DataValidation/Private/EditorValidator_Localization.cpp
+++ b/Engine/Plugins/Editor/DataValidation/Source/DataValidation/Private/EditorValidator_Localization.cpp
@@ -23,12 +23,12 @@ UEditorValidator_Localization::UEditorValidator_Localization()
 	bIsEnabled = true;
 }
 
-bool UEditorValidator_Localization::CanValidateAsset_Implementation(UObject* InAsset) const
+bool UEditorValidator_Localization::CanValidateAsset_Implementation(const FAssetData& AssetData, UObject* InAsset, FDataValidationContext& InContext) const
 {
 	return true;
 }
 
-EDataValidationResult UEditorValidator_Localization::ValidateLoadedAsset_Implementation(UObject* InAsset, TArray<FText>& ValidationErrors)
+EDataValidationResult UEditorValidator_Localization::ValidateLoadedAsset_Implementation(const FAssetData& AssetData, UObject* InAsset, FDataValidationContext& InContext)
 {
 	static const FName NAME_AssetToools = "AssetTools";
 	IAssetTools& AssetTools = FModuleManager::Get().LoadModuleChecked<FAssetToolsModule>(NAME_AssetToools).Get();
@@ -53,7 +53,8 @@ EDataValidationResult UEditorValidator_Localization::ValidateLoadedAsset_Impleme
 				// Can this type of asset be localized?
 				if (!AssetTools.CanLocalize(InAsset->GetClass()))
 				{
-					AssetFails(InAsset, FText::Format(LOCTEXT("LocalizationError_LocalizedAssetTypeCannotBeLocalized", "Localized asset is of type '{0}', which is not a type that can be localized!"), FText::FromString(InAsset->GetClass()->GetName())), ValidationErrors);
+					AssetFails(InAsset, FText::Format(LOCTEXT("LocalizationError_LocalizedAssetTypeCannotBeLocalized", "Localized asset is of type '{0}', which is not a type that can be localized!"),
+						FText::FromString(InAsset->GetClass()->GetName())));
 					return EDataValidationResult::Invalid;
 				}
 
@@ -82,7 +83,8 @@ EDataValidationResult UEditorValidator_Localization::ValidateLoadedAsset_Impleme
 
 					if (!bIsRedirectorReferringToRedirectedLocation)
 					{
-						AssetFails(InAsset, LOCTEXT("LocalizationError_LocalizedAssetHasSourceRedirector", "Localized asset has a source asset that is a redirector. Did you forget to rename the localized assets too?"), ValidationErrors);
+						AssetFails(InAsset, 
+							LOCTEXT("LocalizationError_LocalizedAssetHasSourceRedirector", "Localized asset has a source asset that is a redirector. Did you forget to rename the localized assets too?"));
 						return EDataValidationResult::Invalid;
 					}
 				}
@@ -90,7 +92,9 @@ EDataValidationResult UEditorValidator_Localization::ValidateLoadedAsset_Impleme
 				// Is the source asset the expected type?
 				if (SourceAssetData.AssetClassPath != InAsset->GetClass()->GetClassPathName())
 				{
-					AssetFails(InAsset, FText::Format(LOCTEXT("LocalizationError_LocalizedTypeMismatchWithSourceAsset", "Localized asset is of type '{0}', but its source asset is of type '{1}'. A localized asset must have the same type as its source asset!"), FText::FromString(InAsset->GetClass()->GetPathName()), FText::FromString(SourceAssetData.AssetClassPath.ToString())), ValidationErrors);
+					AssetFails(InAsset, FText::Format(LOCTEXT("LocalizationError_LocalizedTypeMismatchWithSourceAsset", 
+							"Localized asset is of type '{0}', but its source asset is of type '{1}'. A localized asset must have the same type as its source asset!"),
+							FText::FromString(InAsset->GetClass()->GetPathName()), FText::FromString(SourceAssetData.AssetClassPath.ToString())));
 					return EDataValidationResult::Invalid;
 				}
 			}
@@ -119,7 +123,9 @@ EDataValidationResult UEditorValidator_Localization::ValidateLoadedAsset_Impleme
 							// Can this type of asset be localized?
 							if (!AssetTools.CanLocalize(InAsset->GetClass()))
 							{
-								AssetFails(InAsset, FText::Format(LOCTEXT("LocalizationError_SourceAssetTypeCannotBeLocalized", "Source asset has a localized asset for '{0}', but is of type '{1}' which is not a type that can be localized!"), FText::FromString(CultureName), FText::FromString(InAsset->GetClass()->GetName())), ValidationErrors);
+								AssetFails(InAsset, FText::Format(LOCTEXT("LocalizationError_SourceAssetTypeCannotBeLocalized", 
+									"Source asset has a localized asset for '{0}', but is of type '{1}' which is not a type that can be localized!"),
+									FText::FromString(CultureName), FText::FromString(InAsset->GetClass()->GetName())));
 								return EDataValidationResult::Invalid;
 							}
 
@@ -148,7 +154,9 @@ EDataValidationResult UEditorValidator_Localization::ValidateLoadedAsset_Impleme
 
 								if (!bIsRedirectorReferringToRedirectedLocation)
 								{
-									AssetFails(InAsset, FText::Format(LOCTEXT("LocalizationError_SourceIsRedirector", "Source asset is a redirector but has a localized asset for '{0}'. Did you forget to rename the localized assets too?"), FText::FromString(CultureName)), ValidationErrors);
+									AssetFails(InAsset, FText::Format(LOCTEXT("LocalizationError_SourceIsRedirector", 
+										"Source asset is a redirector but has a localized asset for '{0}'. Did you forget to rename the localized assets too?"), 
+										FText::FromString(CultureName)));
 									return EDataValidationResult::Invalid;
 								}
 							}
@@ -156,7 +164,9 @@ EDataValidationResult UEditorValidator_Localization::ValidateLoadedAsset_Impleme
 							// Is the localized asset the expected type?
 							if (LocalizedAssetData.AssetClassPath != InAsset->GetClass()->GetClassPathName())
 							{
-								AssetFails(InAsset, FText::Format(LOCTEXT("LocalizationError_SourceTypeMismatchWithLocalizedAsset", "Source asset is of type '{0}', but its localized asset for '{1}' is of type '{2}'. A localized asset must have the same type as its source asset!"), FText::FromString(InAsset->GetClass()->GetPathName()), FText::FromString(CultureName), FText::FromString(LocalizedAssetData.AssetClassPath.ToString())), ValidationErrors);
+								AssetFails(InAsset, FText::Format(LOCTEXT("LocalizationError_SourceTypeMismatchWithLocalizedAsset", 
+									"Source asset is of type '{0}', but its localized asset for '{1}' is of type '{2}'. A localized asset must have the same type as its source asset!"), 
+									FText::FromString(InAsset->GetClass()->GetPathName()), FText::FromString(CultureName), FText::FromString(LocalizedAssetData.AssetClassPath.ToString())));
 								return EDataValidationResult::Invalid;
 							}
 						}
diff --git a/Engine/Plugins/Editor/DataValidation/Source/DataValidation/Private/WorldPartitionChangelistValidator.cpp b/Engine/Plugins/Editor/DataValidation/Source/DataValidation/Private/WorldPartitionChangelistValidator.cpp
--- a/Engine/Plugins/Editor/DataValidation/Source/DataValidation/Private/WorldPartitionChangelistValidator.cpp
+++ b/Engine/Plugins/Editor/DataValidation/Source/DataValidation/Private/WorldPartitionChangelistValidator.cpp
@@ -23,43 +23,38 @@
 
 #define LOCTEXT_NAMESPACE "WorldPartitionChangelistValidation"
 
-bool UWorldPartitionChangelistValidator::CanValidateAsset_Implementation(UObject* InAsset) const
+bool UWorldPartitionChangelistValidator::CanValidateAsset_Implementation(const FAssetData& AssetData, UObject* InAsset, FDataValidationContext& InContext) const
 {
 	return (InAsset != nullptr) && (UDataValidationChangelist::StaticClass() == InAsset->GetClass());
 }
 
-EDataValidationResult UWorldPartitionChangelistValidator::ValidateLoadedAsset_Implementation(UObject* InAsset, TArray<FText>& ValidationErrors)
+EDataValidationResult UWorldPartitionChangelistValidator::ValidateLoadedAsset_Implementation(const FAssetData& AssetData, UObject* InAsset, FDataValidationContext& InContext)
 {
 	TRACE_CPUPROFILER_EVENT_SCOPE(UWorldPartitionChangelistValidator::ValidateLoadedAsset_Implementation);
 
 	UDataValidationChangelist* ChangeList = CastChecked<UDataValidationChangelist>(InAsset);
 	
-	Errors = &ValidationErrors;
-	
-	EDataValidationResult ValidateActorAndDataLayerResult = ValidateActorsAndDataLayersFromChangeList(ChangeList);
-
-	if (ValidateActorAndDataLayerResult == EDataValidationResult::Invalid)
-	{
-		AssetFails(InAsset, LOCTEXT("WorldPartitionValidationFail", "This changelist contains modifications that aren't valid at the world partition level. Please see revision control log and correct the errors."), ValidationErrors);
-	}
-	else
+	ValidateActorsAndDataLayersFromChangeList(ChangeList);
+	switch (GetValidationResult())
 	{
-		AssetPasses(InAsset);
+		case EDataValidationResult::Invalid:
+			AssetFails(InAsset, LOCTEXT("WorldPartitionValidationFail", 
+				"This changelist contains modifications that aren't valid at the world partition level. Please see the message log for the errors preceding this message."));
+			break;
+		case EDataValidationResult::Valid:
+		case EDataValidationResult::NotValidated:
+			AssetPasses(InAsset);
+			break;
 	}
 
-	return Errors->IsEmpty() ? EDataValidationResult::Valid : EDataValidationResult::Invalid;
+	return GetValidationResult();
 }
 
 // Extract all Actors/Map from Changelist (in OFPA this should be one Actor per Package, and we'll discard all Actors from non WorldPartition maps)
 // and add them to a Map of World->Files[] so that we can do one validation per world. Once Worlds are identified, we either the UActorDescContainer 
 // from memory (if loaded) or request it to be loaded, we then build a Set of objects that interest us from the Actors in the CL 
-EDataValidationResult UWorldPartitionChangelistValidator::ValidateActorsAndDataLayersFromChangeList(UDataValidationChangelist* Changelist)
+void UWorldPartitionChangelistValidator::ValidateActorsAndDataLayersFromChangeList(UDataValidationChangelist* Changelist)
 {
-	int32 NumErrorsOnEntry = Errors->Num();
-
-	ISourceControlProvider& SourceControlProvider = ISourceControlModule::Get().GetProvider();
-	FSourceControlChangelistStatePtr ChangelistState = SourceControlProvider.GetState(Changelist->Changelist->AsShared(), EStateCacheUsage::Use);
-
 	// Figure out which world(s) these actors are in and split the files per world, and return the actor's native class
 	TMap<FTopLevelAssetPath, TSet<FAssetData>> MapToActorsFiles;
 	auto TryAssociateActorToMap = [&MapToActorsFiles](const FAssetData& AssetData) -> UClass*
@@ -94,57 +89,47 @@ EDataValidationResult UWorldPartitionChangelistValidator::ValidateActorsAndDataL
 	
 	FAssetRegistryModule& AssetRegistryModule = FModuleManager::LoadModuleChecked<FAssetRegistryModule>(FName("AssetRegistry"));
 	IAssetRegistry& AssetRegistry = AssetRegistryModule.Get();
-	for (const FSourceControlStateRef& File : ChangelistState->GetFilesStates())
+	for (FName PackageName : Changelist->ModifiedPackageNames)
 	{		
-		// Skip deleted files since we're not validating references in this validator 
-		if (File->IsDeleted())
+		TArray<FAssetData> PackageAssetsData;
+		USourceControlHelpers::GetAssetDataFromPackage(PackageName.ToString(), PackageAssetsData);
+		
+		for (FAssetData& AssetData : PackageAssetsData)
 		{
-			continue;
-		}
-
-		FString PackageName;
-		if (FPackageName::TryConvertFilenameToLongPackageName(File->GetFilename(), PackageName))
-		{						
-			TArray<FAssetData> PackageAssetsData;
-			USourceControlHelpers::GetAssetDataFromPackage(PackageName, PackageAssetsData);
-			
-			for (FAssetData& AssetData : PackageAssetsData)
+			if (UClass* ActorNativeClass = TryAssociateActorToMap(AssetData))
 			{
-				if (UClass* ActorNativeClass = TryAssociateActorToMap(AssetData))
-				{
-					SubmittingWorldDataLayers = ActorNativeClass->IsChildOf<AWorldDataLayers>();
-				}
-				else if (UClass* AssetClass = AssetData.GetClass())
+				SubmittingWorldDataLayers = ActorNativeClass->IsChildOf<AWorldDataLayers>();
+			}
+			else if (UClass* AssetClass = AssetData.GetClass())
+			{
+				if (AssetClass->IsChildOf<UDataLayerAsset>())
 				{
-					if (AssetClass->IsChildOf<UDataLayerAsset>())
-					{
-						TArray<FName> ReferencerNames;
-						AssetRegistry.GetReferencers(AssetData.PackageName, ReferencerNames, UE::AssetRegistry::EDependencyCategory::All);
-			
-						FARFilter Filter;
-						Filter.bIncludeOnlyOnDiskAssets = true;
-						Filter.PackageNames = MoveTemp(ReferencerNames);
+					TArray<FName> ReferencerNames;
+					AssetRegistry.GetReferencers(AssetData.PackageName, ReferencerNames, UE::AssetRegistry::EDependencyCategory::All);
+		
+					FARFilter Filter;
+					Filter.bIncludeOnlyOnDiskAssets = true;
+					Filter.PackageNames = MoveTemp(ReferencerNames);
 
-						TArray<FAssetData> DataLayerReferencers;
-						AssetRegistry.GetAssets(Filter, DataLayerReferencers);
+					TArray<FAssetData> DataLayerReferencers;
+					AssetRegistry.GetAssets(Filter, DataLayerReferencers);
 
-						for (const FAssetData& DataLayerReferencer : DataLayerReferencers)
+					for (const FAssetData& DataLayerReferencer : DataLayerReferencers)
+					{
+						UClass* ReferencerAssetClass = DataLayerReferencer.GetClass();
+						if (ReferencerAssetClass && ReferencerAssetClass->IsChildOf<AWorldDataLayers>())
 						{
-							UClass* ReferencerAssetClass = DataLayerReferencer.GetClass();
-							if (ReferencerAssetClass && ReferencerAssetClass->IsChildOf<AWorldDataLayers>())
-							{
-								TryAssociateActorToMap(DataLayerReferencer);
-							}
+							TryAssociateActorToMap(DataLayerReferencer);
 						}
-			
-						RelevantDataLayerAssets.Add(AssetData.PackageName.ToString());
 					}
-					else if (AssetClass->IsChildOf<UWorld>())
+		
+					RelevantDataLayerAssets.Add(AssetData.PackageName.ToString());
+				}
+				else if (AssetClass->IsChildOf<UWorld>())
+				{
+					if (ULevel::GetIsLevelPartitionedFromPackage(PackageName))
 					{
-						if (ULevel::GetIsLevelPartitionedFromPackage(*PackageName))
-						{
-							MapToActorsFiles.FindOrAdd(AssetData.GetSoftObjectPath().GetAssetPath());
-						}
+						MapToActorsFiles.FindOrAdd(AssetData.GetSoftObjectPath().GetAssetPath());
 					}
 				}
 			}
@@ -185,6 +170,8 @@ EDataValidationResult UWorldPartitionChangelistValidator::ValidateActorsAndDataL
 		// Find/Load the ActorDescContainer
 		UWorld* World = FindObject<UWorld>(nullptr, *MapPath.ToString(), true);
 		
+		TGuardValue<UObject*> GuardCurrentAsset(CurrentAsset, World);
+		
 		FActorDescContainerCollection ContainersToValidate;
 		for (const FAssetData& ActorData : ActorsData)
 		{
@@ -238,14 +225,6 @@ EDataValidationResult UWorldPartitionChangelistValidator::ValidateActorsAndDataL
 			UWorldPartition::CheckForErrors(Params);
 		});
 	}
-
-	if (Errors->Num() > NumErrorsOnEntry)
-	{
-		return EDataValidationResult::Invalid;
-	}
-		
-	return EDataValidationResult::Valid;
-	
 }
 
 bool UWorldPartitionChangelistValidator::Filter(const FWorldPartitionActorDescView& ActorDescView)
@@ -278,7 +257,7 @@ void UWorldPartitionChangelistValidator::OnInvalidRuntimeGrid(const FWorldPartit
 											FText::FromString(GetFullActorName(ActorDescView)), 
 											FText::FromName(GridName));
 
-		Errors->Add(CurrentError);
+		AssetFails(CurrentAsset, CurrentError);
 	}
 }
 
@@ -290,7 +269,7 @@ void UWorldPartitionChangelistValidator::OnInvalidReference(const FWorldPartitio
 											FText::FromString(GetFullActorName(ActorDescView)), 
 											FText::FromString(ReferenceActorDescView ? GetFullActorName(*ReferenceActorDescView) : ReferenceGuid.ToString()));
 
-		Errors->Add(CurrentError);
+		AssetFails(CurrentAsset, CurrentError);
 	}
 }
 
@@ -310,7 +289,7 @@ void UWorldPartitionChangelistValidator::OnInvalidReferenceGridPlacement(const F
 												FText::FromString(ReferenceActorDescView.GetIsSpatiallyLoaded() ? *SpatiallyLoadedActor : *NonSpatiallyLoadedActor),
 												FText::FromString(GetFullActorName(ReferenceActorDescView)));
 
-			Errors->Add(CurrentError);
+		AssetFails(CurrentAsset, CurrentError);
 		}
 	}
 }
@@ -323,7 +302,7 @@ void UWorldPartitionChangelistValidator::OnInvalidReferenceDataLayers(const FWor
 											FText::FromString(GetFullActorName(ActorDescView)),
 											FText::FromString(GetFullActorName(ReferenceActorDescView)));
 
-		Errors->Add(CurrentError);
+		AssetFails(CurrentAsset, CurrentError);
 	}
 }
 
@@ -335,7 +314,7 @@ void UWorldPartitionChangelistValidator::OnInvalidReferenceRuntimeGrid(const FWo
 			FText::FromString(GetFullActorName(ActorDescView)),
 			FText::FromString(GetFullActorName(ReferenceActorDescView)));
 
-		Errors->Add(CurrentError);
+		AssetFails(CurrentAsset, CurrentError);
 	}
 }
 
@@ -346,7 +325,7 @@ void UWorldPartitionChangelistValidator::OnInvalidReferenceLevelScriptStreamed(c
 		FText CurrentError = FText::Format(LOCTEXT("DataValidation.Changelist.WorldPartition.InvalidReferenceLevelScriptStreamed", "Level script blueprint references streamed actor {0}."),
 											FText::FromString(GetFullActorName(ActorDescView)));
 		
-		Errors->Add(CurrentError);
+		AssetFails(CurrentAsset, CurrentError);
 	}
 }
 
@@ -357,7 +336,7 @@ void UWorldPartitionChangelistValidator::OnInvalidReferenceLevelScriptDataLayers
 		FText CurrentError = FText::Format(LOCTEXT("DataValidation.Changelist.WorldPartition.InvalidReferenceLevelScriptDataLayers", "Level script blueprint references streamed actor {0} with a non empty set of data layers."),
 											FText::FromString(GetFullActorName(ActorDescView)));
 
-		Errors->Add(CurrentError);
+		AssetFails(CurrentAsset, CurrentError);
 	}
 }
 
@@ -368,7 +347,7 @@ void UWorldPartitionChangelistValidator::OnInvalidReferenceDataLayerAsset(const
 		FText CurrentError = FText::Format(LOCTEXT("DataValidation.Changelist.WorldPartition.InvalidDataLayerAsset", "Data layer {0} has no data layer asset."),
 			FText::FromName(DataLayerInstance->GetDataLayerFName()));
 
-		Errors->Add(CurrentError);
+		AssetFails(CurrentAsset, CurrentError);
 	}
 }
 
@@ -384,7 +363,7 @@ void UWorldPartitionChangelistValidator::OnDataLayerHierarchyTypeMismatch(const
 			FText::FromString(Parent->GetDataLayerFullName()),
 			UEnum::GetDisplayValueAsText(Parent->GetType()));
 	
-		Errors->Add(CurrentError);
+		AssetFails(CurrentAsset, CurrentError);
 	}
 }
 
@@ -399,7 +378,7 @@ void UWorldPartitionChangelistValidator::OnDataLayerAssetConflict(const UDataLay
 			FText::FromName(ConflictingDataLayerInstance->GetDataLayerFName()),
 			FText::FromString(DataLayerInstance->GetAsset()->GetFullName()));
 
-		Errors->Add(CurrentError);
+		AssetFails(CurrentAsset, CurrentError);
 	}
 }
 
@@ -420,7 +399,7 @@ void UWorldPartitionChangelistValidator::OnLevelInstanceInvalidWorldAsset(const
 			CurrentError = FText::Format(LOCTEXT("DataValidation.Changelist.WorldPartition.LevelInstanceInvalidWorldAsset", "Level instance {0} has an invalid world asset {1}."),
 				FText::FromString(GetFullActorName(ActorDescView)), 
 				FText::FromName(WorldAsset));
-			Errors->Add(CurrentError);
+			AssetFails(CurrentAsset, CurrentError);
 			break;
 		case ELevelInstanceInvalidReason::WorldAssetNotUsingExternalActors:
 			// Not a validation error
@@ -435,7 +414,7 @@ void UWorldPartitionChangelistValidator::OnLevelInstanceInvalidWorldAsset(const
 			CurrentError = FText::Format(LOCTEXT("DataValidation.Changelist.WorldPartition.LevelInstanceCircularReference", "Level instance {0} has a circular reference {1}."),
 				FText::FromString(GetFullActorName(ActorDescView)),
 				FText::FromName(WorldAsset));
-			Errors->Add(CurrentError);
+			AssetFails(CurrentAsset, CurrentError);
 			break;
 		};
 	}
@@ -454,7 +433,7 @@ void UWorldPartitionChangelistValidator::OnInvalidHLODLayer(const FWorldPartitio
 											FText::FromString(GetFullActorName(ActorDescView)), 
 											FText::FromString(ActorDescView.GetHLODLayer().ToString()));
 
-		Errors->Add(CurrentError);
+		AssetFails(CurrentAsset, CurrentError);
 	}
 }
 
diff --git a/Engine/Plugins/Editor/DataValidation/Source/DataValidation/Private/WorldPartitionChangelistValidator.h b/Engine/Plugins/Editor/DataValidation/Source/DataValidation/Private/WorldPartitionChangelistValidator.h
--- a/Engine/Plugins/Editor/DataValidation/Source/DataValidation/Private/WorldPartitionChangelistValidator.h
+++ b/Engine/Plugins/Editor/DataValidation/Source/DataValidation/Private/WorldPartitionChangelistValidator.h
@@ -16,16 +16,16 @@ class DATAVALIDATION_API UWorldPartitionChangelistValidator : public UEditorVali
 	GENERATED_BODY()
 
 protected:	
-	TArray<FText>* Errors = nullptr;
 	FTopLevelAssetPath RelevantMap;
 	TSet<FGuid> RelevantActorGuids;
 	TSet<FString> RelevantDataLayerAssets;
+	UObject* CurrentAsset = nullptr;
 	bool SubmittingWorldDataLayers = false;
 
-	virtual bool CanValidateAsset_Implementation(UObject* InAsset) const override;
-	virtual EDataValidationResult ValidateLoadedAsset_Implementation(UObject* InAsset, TArray<FText>& ValidationErrors) override;
+	virtual bool CanValidateAsset_Implementation(const FAssetData& InAssetData, UObject* InAsset, FDataValidationContext& InContext) const override;
+	virtual EDataValidationResult ValidateLoadedAsset_Implementation(const FAssetData& InAssetData, UObject* InAsset, FDataValidationContext& InContext) override;
 
-	EDataValidationResult ValidateActorsAndDataLayersFromChangeList(UDataValidationChangelist* Changelist);
+	void ValidateActorsAndDataLayersFromChangeList(UDataValidationChangelist* Changelist);
 
 	// Return true if this ActorDescView is pertinent to the current ChangeList
 	bool Filter(const FWorldPartitionActorDescView& ActorDescView);
diff --git a/Engine/Plugins/Editor/DataValidation/Source/DataValidation/Public/DataValidationChangelist.h b/Engine/Plugins/Editor/DataValidation/Source/DataValidation/Public/DataValidationChangelist.h
--- a/Engine/Plugins/Editor/DataValidation/Source/DataValidation/Public/DataValidationChangelist.h
+++ b/Engine/Plugins/Editor/DataValidation/Source/DataValidation/Public/DataValidationChangelist.h
@@ -5,6 +5,7 @@
 
 
 #include "ISourceControlChangelist.h"
+#include "ISourceControlChangelistState.h"
 
 #include "DataValidationChangelist.generated.h"
 
@@ -22,20 +23,28 @@ public:
 
 	virtual EDataValidationResult IsDataValid(class FDataValidationContext& Context) const override;
 
-	/** Initializes internal state so the validation can be done */
+	/** Initializes from a list of file states as a pseudo-changelist */
+	void Initialize(TConstArrayView<FSourceControlStateRef> InFileStates);
+	/** Initializes from a changelist reference, querying the state from the provider */
 	void Initialize(FSourceControlChangelistPtr InChangelist);
+	/** Initializes from an already-queried changelist state */
+	void Initialize(FSourceControlChangelistStateRef InChangelistState);
 
 	static void GatherDependencies(const FName& InPackageName, TSet<FName>& OutDependencies);
 	static FString GetPrettyPackageName(const FName& InPackageName);
-public:
-	/** Changelist to validate */
-	FSourceControlChangelistPtr Changelist;
 
-	/** Assets contained in the changelist */
-	//TArray<FAssetData> AssetsInChangelist;
-
-	/** Change description */
-	//FText Description;
+	/** Changelist to validate - may be null if this was constructed from a list of files */
+	FSourceControlChangelistPtr Changelist;
+	
+	// Asset files in the changelist 
+	TArray<FName> ModifiedPackageNames;
+	TArray<FName> DeletedPackageNames;
+	
+	// Non-asset files in the changelist
+	TArray<FString> ModifiedFiles;
+	TArray<FString> DeletedFiles;
+	
+	FText Description;
 };
 
 #if UE_ENABLE_INCLUDE_ORDER_DEPRECATED_IN_5_2
diff --git a/Engine/Plugins/Editor/DataValidation/Source/DataValidation/Public/DataValidationModule.h b/Engine/Plugins/Editor/DataValidation/Source/DataValidation/Public/DataValidationModule.h
--- a/Engine/Plugins/Editor/DataValidation/Source/DataValidation/Public/DataValidationModule.h
+++ b/Engine/Plugins/Editor/DataValidation/Source/DataValidation/Public/DataValidationModule.h
@@ -4,38 +4,15 @@
 
 #include "Modules/ModuleManager.h"
 #include "UObject/ObjectMacros.h"
+#include "Misc/DataValidation.h"
 
 #if UE_ENABLE_INCLUDE_ORDER_DEPRECATED_IN_5_2
 #include "AssetRegistry/AssetData.h"
 #include "CoreMinimal.h"
 #endif
 
-#include "DataValidationModule.generated.h"
-
 struct FAssetData;
 
-UENUM(BlueprintType)
-enum class EDataValidationUsecase : uint8
-{
-	/** No usecase specified */
-	None = 0,
-
-	/** Triggered on user's demand */
-	Manual,
-
-	/** A commandlet invoked the validation */
-	Commandlet,
-
-	/** Saving a package triggered the validation */
-	Save,
-
-	/** Submit dialog triggered the validation */
-	PreSubmit,
-
-	/** Triggered by blueprint or c++ */
-	Script,
-};
-
 /**
  * The public interface to this module
  */
diff --git a/Engine/Plugins/Editor/DataValidation/Source/DataValidation/Public/DataValidationSettings.h b/Engine/Plugins/Editor/DataValidation/Source/DataValidation/Public/DataValidationSettings.h
new file mode 100644
--- /dev/null
+++ b/Engine/Plugins/Editor/DataValidation/Source/DataValidation/Public/DataValidationSettings.h
@@ -0,0 +1,26 @@
+#pragma once
+
+#include "Engine/DeveloperSettings.h"
+
+#include "DataValidationSettings.generated.h"
+
+/**
+ * Project-wide settings for data validation
+ */
+UCLASS(config = Editor)
+class DATAVALIDATION_API UDataValidationSettings : public UDeveloperSettings
+{
+	GENERATED_BODY()
+
+public:
+	/** Default constructor that sets up CDO properties */
+	UDataValidationSettings();
+
+	/** Whether or not to validate assets on save */
+	UPROPERTY(EditAnywhere, config, Category="Data Validation")
+	uint32 bValidateOnSave : 1 = true;
+	
+	/** Whether or not to load & validate assets in changelists by default */
+	UPROPERTY(EditAnywhere, config, Category="Data Validation")
+	uint32 bLoadAssetsWhenValidatingChangelists : 1 = true;
+};
diff --git a/Engine/Plugins/Editor/DataValidation/Source/DataValidation/Public/DirtyFilesChangelistValidator.h b/Engine/Plugins/Editor/DataValidation/Source/DataValidation/Public/DirtyFilesChangelistValidator.h
--- a/Engine/Plugins/Editor/DataValidation/Source/DataValidation/Public/DirtyFilesChangelistValidator.h
+++ b/Engine/Plugins/Editor/DataValidation/Source/DataValidation/Public/DirtyFilesChangelistValidator.h
@@ -16,6 +16,6 @@ class DATAVALIDATION_API UDirtyFilesChangelistValidator : public UEditorValidato
 	GENERATED_BODY()
 
 protected:
-	virtual bool CanValidateAsset_Implementation(UObject* InAsset) const override;
-	virtual EDataValidationResult ValidateLoadedAsset_Implementation(UObject* InAsset, TArray<FText>& ValidationErrors) override;
+	virtual bool CanValidateAsset_Implementation(const FAssetData& InAssetData, UObject* InAsset, FDataValidationContext& InContext) const override;
+	virtual EDataValidationResult ValidateLoadedAsset_Implementation(const FAssetData& InAssetData, UObject* InAsset, FDataValidationContext& InContext) override;
 };
diff --git a/Engine/Plugins/Editor/DataValidation/Source/DataValidation/Public/EditorValidatorBase.h b/Engine/Plugins/Editor/DataValidation/Source/DataValidation/Public/EditorValidatorBase.h
--- a/Engine/Plugins/Editor/DataValidation/Source/DataValidation/Public/EditorValidatorBase.h
+++ b/Engine/Plugins/Editor/DataValidation/Source/DataValidation/Public/EditorValidatorBase.h
@@ -2,12 +2,15 @@
 
 #pragma once
 
-
-
 #include "UObject/Package.h"
+
 #include "EditorValidatorBase.generated.h"
 
 enum class EDataValidationUsecase : uint8;
+namespace EMessageSeverity { enum Type : int; };
+struct FAssetData;
+class FDataValidationContext;
+class FTokenizedMessage;
 
 /*
 * The EditorValidatorBase is a class which verifies that an asset meets a specific ruleset.
@@ -18,37 +21,72 @@ enum class EDataValidationUsecase : uint8;
 * C++ and Blueprint validators will be gathered on editor start, while python validators need
 * to register themselves
 */
-UCLASS(Abstract, Blueprintable, meta = (ShowWorldContextPin))
+UCLASS(Abstract, Blueprintable, meta = (ShowWorldContextPin), Config=Editor)
 class DATAVALIDATION_API UEditorValidatorBase : public UObject
 {
 	GENERATED_BODY()
 
 public:
 	UEditorValidatorBase();
-
+	
+	// Validation entry point - will dispatch to native and/or BP validation as appropriate and return results.
+	// If validator doesn't want to validate this asset/in this context, it will return NotValidated which is acceptable.
+	EDataValidationResult ValidateLoadedAsset(const FAssetData& InAssetData, UObject* InAsset, FDataValidationContext& Context);
+
+	/** 
+	 * Override this to determine whether or not you can use this validator given this context.
+	 * Context can be used to add errors if validation cannot be performed because of some issue 
+	 */
+	virtual bool CanValidateAsset_Implementation(const FAssetData& InAssetData, UObject* InObject, FDataValidationContext& InContext) const
+	{
+		return false;
+	}
+	
+	/** Override this to validate in C++ with access to FDataValidationContext */
+	virtual EDataValidationResult ValidateLoadedAsset_Implementation(const FAssetData& InAssetData, UObject* InAsset, FDataValidationContext& Context)
+	{
+		return EDataValidationResult::NotValidated;
+	}
+	
 	/** Override this to determine whether or not you can use this validator given this usecase */
-	UFUNCTION(BlueprintNativeEvent, Category = "Asset Validation")
-	bool CanValidate(const EDataValidationUsecase InUsecase) const;
-
+	UFUNCTION(BlueprintNativeEvent, Category = "Asset Validation", DisplayName="Can Validate")
+	bool K2_CanValidate(const EDataValidationUsecase InUsecase) const;
+	
 	/** Override this to determine whether or not you can validate a given asset with this validator */
-	UFUNCTION(BlueprintNativeEvent, Category = "Asset Validation")
-	bool CanValidateAsset(UObject* InAsset) const;
-
-	UFUNCTION(BlueprintNativeEvent, Category = "Asset Validation")
-	EDataValidationResult ValidateLoadedAsset(UObject* InAsset, UPARAM(ref) TArray<FText>& ValidationErrors);
-
+	UFUNCTION(BlueprintNativeEvent, Category = "Asset Validation", DisplayName="Can Validate Asset")
+	bool K2_CanValidateAsset(UObject* InAsset) const;
+	
+	/** Override this in blueprint to validate assets */
+	UFUNCTION(BlueprintNativeEvent, Category = "Asset Validation", DisplayName="Validate Loaded Asset")
+	EDataValidationResult K2_ValidateLoadedAsset(UObject* InAsset);
+
+	// Declaration for above BlueprintNativeEvent
+	virtual EDataValidationResult K2_ValidateLoadedAsset_Implementation(UObject* InAsset);
+	
+	/** Marks the validation as failed and adds an error message. */
 	UFUNCTION(BlueprintCallable, Category = "Asset Validation")
-	void AssetFails(UObject* InAsset, const FText& InMessage, UPARAM(ref) TArray<FText>& ValidationErrors);
+	void AssetFails(UObject* InAsset, const FText& InMessage);
+
+	UE_DEPRECATED("5.4", "This function has been replaced by the version AssetFails(UObject*, const FText&)")
+	void AssetFails(UObject* InAsset, const FText& InMessage, TArray<FText>& InOutErrors);
 
+	/** Marks the validation as successful. Failure to call this will report the validator as not having checked the asset. */
 	UFUNCTION(BlueprintCallable, Category = "Asset Validation")
 	void AssetPasses(UObject* InAsset);
 
+	/** Adds a message to this validation but doesn't mark it as failed. */
 	UFUNCTION(BlueprintCallable, Category = "Asset Validation")
 	void AssetWarning(UObject* InAsset, const FText& InMessage);
+	
+	/** 
+	 * Add a tokenized message to the validation results. If the severity is error, marks the validation as failed.
+	 */
+	TSharedRef<FTokenizedMessage> AssetMessage(const FAssetData& InAssetData, EMessageSeverity::Type InSeverity, const FText& InText = {});
+	TSharedRef<FTokenizedMessage> AssetMessage(EMessageSeverity::Type InSeverity, const FText& InText = {});
 
 	virtual bool IsEnabled() const
 	{
-		return bIsEnabled;
+		return bIsEnabled && !bIsConfigDisabled;
 	}
 
 	void ResetValidationState();
@@ -65,6 +103,26 @@ public:
 	}
 
 	const TArray<FText>& GetAllWarnings() const;
+	
+	void AddLegacyValidationErrors(TArray<FText> InErrors);
+	EDataValidationResult ExtractValidationState(FDataValidationContext& InOutContext) const;
+
+	// Legacy entry points from old interface 
+	UE_DEPRECATED("5.4", "CanValidate_Implementation(const EDataValidationUsecase InUsecase) is deprecated, override CanValidateAsset_Implementation(UObject* InObject, FDataValidationContext& InContext) instead")
+	virtual bool CanValidate_Implementation(const EDataValidationUsecase InUsecase) const 
+	{
+		return false; 
+	}
+	UE_DEPRECATED("5.4", "CanValidateAsset_Implementation(UObject* InAsset) is deprecated, override CanValidateAsset_Implementation(UObject* InObject, FDataValidationContext& InContext) instead")
+	virtual bool CanValidateAsset_Implementation(UObject* InAsset) const 
+	{
+		 return true; 
+	}
+	UE_DEPRECATED("5.4", "ValidateLoadedAsset_Implementation(UObject* InAsset, TArray<FText>& ValidationErrors) is deprecated, override ValidateLoadedAsset_Implementation(UObject* InAsset, FDataValidationContext& Context) instead")
+	virtual EDataValidationResult ValidateLoadedAsset_Implementation(UObject* InAsset, TArray<FText>& ValidationErrors) 
+	{
+		return EDataValidationResult::NotValidated;
+	}
 
 protected:
 	void LogElapsedTime(FFormatNamedArguments &Arguments);
@@ -72,12 +130,23 @@ protected:
 protected:
 	UPROPERTY(EditAnywhere, Category = "Asset Validation", meta = (BlueprintProtected = "true"))
 	bool bIsEnabled;
+	
+	// Allows disabling validators from config without checking in assets/modifying code.
+	UPROPERTY(Config)
+	bool bIsConfigDisabled;
 
 private:
 	EDataValidationResult ValidationResult;
 	TArray<FText> AllWarnings;
+	TArray<FText> AllErrors;
+	TArray<TSharedRef<FTokenizedMessage>> AllMessages;
 
 	FDateTime ValidationTime;
+	
+	UPROPERTY()
+	UObject* CurrentObjectBeingValidated;
+
+	const FAssetData* CurrentAssetBeingValidated;
 };
 
 
diff --git a/Engine/Plugins/Editor/DataValidation/Source/DataValidation/Public/EditorValidatorHelpers.h b/Engine/Plugins/Editor/DataValidation/Source/DataValidation/Public/EditorValidatorHelpers.h
new file mode 100644
--- /dev/null
+++ b/Engine/Plugins/Editor/DataValidation/Source/DataValidation/Public/EditorValidatorHelpers.h
@@ -0,0 +1,117 @@
+// Copyright Epic Games, Inc. All Rights Reserved.
+
+#pragma once
+
+#include "Containers/ContainersFwd.h"
+#include "Delegates/Delegate.h"
+#include "Templates/Tuple.h"
+#include "UObject/NameTypes.h"
+
+class FMessageLog;
+class FText;
+struct FAssetData;
+class FDataValidationContext;
+class FWildcardString;
+class UObject;
+
+namespace EMessageSeverity
+{
+    enum Type : int;
+}
+
+namespace UE::DataValidation
+{
+    extern DATAVALIDATION_API const FName MessageLogName; 
+    
+    // Utility to capture log messages from all threads during data validation so they can be filtered
+    // and passed as validation errors if desired.
+    struct FLogMessageGathererImpl;
+    struct DATAVALIDATION_API FScopedLogMessageGatherer
+    {
+        FScopedLogMessageGatherer(bool bEnabled = true);
+        ~FScopedLogMessageGatherer();
+
+        FScopedLogMessageGatherer(const FScopedLogMessageGatherer&) = delete;
+        FScopedLogMessageGatherer& operator=(const FScopedLogMessageGatherer&) = delete;
+
+        FScopedLogMessageGatherer(FScopedLogMessageGatherer&&) = delete;
+        FScopedLogMessageGatherer& operator=(FScopedLogMessageGatherer&&) = delete;
+
+        /** Returns the current gather constructed on this thread in an enclosing scope if any */
+        static FScopedLogMessageGatherer* GetCurrentThreadGatherer();
+        
+        /** 
+         * Stop gathering logs and return the warnings and errors gathered so far
+         */
+        void Stop(TArray<FString>& OutWarnings, TArray<FString>& OutErrors);
+
+        /** 
+         * Tell this scope to ignore certain log messages 
+         * @return A tuple that can be passed to RemoveIgnoreCategories to reverse this 
+         */
+        TTuple<int32, int32> AddIgnoreCategories(TConstArrayView<FName> NewCategories);
+
+        /** 
+         * Tell this scope to stop ignoring certain log messages.
+         * @param Range The return value from a previous call to AddIgnoreCategories 
+         */
+        void RemoveIgnoreCategories(TTuple<int32, int32> Range);
+
+        /** 
+         * Tell this scope to ignore certain log messages 
+         * @return A tuple that can be passed to RemoveIgnorePatterns to reverse this 
+         */
+        TTuple<int32, int32> AddIgnorePatterns(TConstArrayView<FWildcardString> NewPatterns);
+
+        /** 
+         * Tell this scope to stop ignoring certain log messages.
+         * @param Range The return value from a previous call to AddIgnorePatterns 
+         */
+        void RemoveIgnorePatterns(TTuple<int32, int32> Range);
+        
+        /** 
+         * Delegate that may be bound by a project to set a default list of log categories to ignore 
+         * during asset validation
+         */
+        static TDelegate<TArray<FName>(void)>& GetDefaultIgnoreCategoriesDelegate();
+
+        /**
+         * Delegate that may be bound by a project to set a default list of log patterns to ignore 
+         * during asset validation
+         */
+        static TDelegate<TArray<FWildcardString>(void)>& GetDefaultIgnorePatternsDelegate();
+        
+    protected:
+        FScopedLogMessageGatherer* Previous = nullptr;
+        FLogMessageGathererImpl* Impl;
+    };
+    
+    // Utility to tell the closest enclosing FScopedLogMessageGatherer to ignore certain messages
+    struct DATAVALIDATION_API FScopedIgnoreLogMessages
+    {
+        FScopedIgnoreLogMessages(TConstArrayView<FWildcardString> Patterns, TConstArrayView<FName> Categories);
+        ~FScopedIgnoreLogMessages();
+        
+    protected:
+        TTuple<int32, int32> PatternsRange = {0, 0};
+        TTuple<int32, int32> CategoriesRange = {0, 0};
+    };
+
+    // Helper to add asset validation messages from a validation context to a message log
+    DATAVALIDATION_API void AddAssetValidationMessages(
+        const FAssetData& ForAsset,
+        FMessageLog& Log,
+        const FDataValidationContext& InContext);
+
+    // Helper to add asset validation messages from a validation context to a message log
+    DATAVALIDATION_API void AddAssetValidationMessages(
+        FMessageLog& Log,
+        const FDataValidationContext& InContext);
+
+    // Helper to format asset validation messages where the messages may contain references to assets that should be replaced by clickable tokens.
+    DATAVALIDATION_API void AddAssetValidationMessages(
+        const FAssetData& ForAsset,
+        FMessageLog& Log,
+        EMessageSeverity::Type Severity,
+        TConstArrayView<FText> Messages);
+}
\ No newline at end of file
diff --git a/Engine/Plugins/Editor/DataValidation/Source/DataValidation/Public/EditorValidatorSubsystem.h b/Engine/Plugins/Editor/DataValidation/Source/DataValidation/Public/EditorValidatorSubsystem.h
--- a/Engine/Plugins/Editor/DataValidation/Source/DataValidation/Public/EditorValidatorSubsystem.h
+++ b/Engine/Plugins/Editor/DataValidation/Source/DataValidation/Public/EditorValidatorSubsystem.h
@@ -5,17 +5,24 @@
 #include "EditorSubsystem.h"
 
 #include "DataValidationModule.h"
+#include "DataValidationSettings.h"
 
 #include "UObject/Package.h"
 #include "EditorValidatorSubsystem.generated.h"
 
 class UEditorValidatorBase;
+class UDataValidationChangelist;
 struct FDirectoryPath;
-
+class FMessageLog;
 struct FAssetData;
 class ISourceControlChangelist;
 typedef TSharedPtr<class ISourceControlChangelist, ESPMode::ThreadSafe> FSourceControlChangelistPtr;
 
+namespace EMessageSeverity
+{
+	enum Type : int;
+}
+
 DECLARE_LOG_CATEGORY_EXTERN(LogContentValidation, Log, All);
 
 USTRUCT(BlueprintType)
@@ -37,6 +44,9 @@ struct DATAVALIDATION_API FValidateAssetsDetails
 
 	/** Validation Warnings */
 	TArray<FText> ValidationWarnings;
+	
+	/** Rich validation messages including tokens that can generate images or hyperlinks */
+	TArray<TSharedRef<FTokenizedMessage>> ValidationMessages;
 };
 
 USTRUCT(BlueprintType)
@@ -45,6 +55,10 @@ struct DATAVALIDATION_API FValidateAssetsResults
 	GENERATED_BODY()
 
 	FValidateAssetsResults() = default;
+	
+	/** Total amount of assets that were gathered to validate. */
+	UPROPERTY(VisibleAnywhere, BlueprintReadOnly, Category = "Validation")
+	int NumRequested = 0;
 
 	/** Amount of tested assets */
 	UPROPERTY(VisibleAnywhere, BlueprintReadOnly, Category = "Validation")
@@ -69,6 +83,10 @@ struct DATAVALIDATION_API FValidateAssetsResults
 	/** Amount of assets that could not be validated */
 	UPROPERTY(VisibleAnywhere, BlueprintReadOnly, Category = "Validation")
 	int NumUnableToValidate = 0;
+	
+	/** True if FValidateAssetSettings::MaxAssetsToValidation was reached */
+	UPROPERTY(VisibleAnywhere, BlueprintReadOnly, Category = "Validation")
+	bool bAssetLimitReached = false;
 
 	/** 
 	 * Per asset details
@@ -84,7 +102,8 @@ struct DATAVALIDATION_API FValidateAssetsSettings
 {
 	GENERATED_BODY()
 
-	FValidateAssetsSettings() = default;
+	FValidateAssetsSettings();
+	~FValidateAssetsSettings();
 
 	/** If true, will not validate files in excluded directories */
 	UPROPERTY(EditAnywhere, BlueprintReadWrite, Category = "Validation")
@@ -103,34 +122,50 @@ struct DATAVALIDATION_API FValidateAssetsSettings
 	EDataValidationUsecase ValidationUsecase = EDataValidationUsecase::None;
 
 	/** If false, unloaded assets will get skipped from validation. */
+	UPROPERTY(EditAnywhere, BlueprintReadWrite, Category = "Validation")
 	bool bLoadAssetsForValidation = true;
-};
-
-/**
-* Implements the settings for Data Validation 
-*/
-UCLASS(config = Editor)
-class DATAVALIDATION_API UDataValidationSettings : public UObject
-{
-	GENERATED_BODY()
+	
+	/** If true, captures log warnings and errors from loading assets for validation and reports them as validation results */
+	UPROPERTY(EditAnywhere, BlueprintReadWrite, Category = "Validation")
+	bool bCaptureAssetLoadLogs = true;
 
-public:
-	/** Default constructor that sets up CDO properties */
-	UDataValidationSettings();
+	/** If true, captures log warnings and errors from other operations that happen during validation and adds them to validation results */
+	UPROPERTY(EditAnywhere, BlueprintReadWrite, Category = "Validation")
+	bool bCaptureLogsDuringValidation = true;
+	
+	/** Maximum number of assets to attempt to validate */
+	UPROPERTY(EditAnywhere, BlueprintReadWrite, Category = "Validation")
+	int32 MaxAssetsToValidate = MAX_int32;
+	
+	/** Should we validate referencers of deleted assets in changelists */
+	UPROPERTY(EditAnywhere, BlueprintReadwrite, Category = "Validation")
+	bool bValidateReferencersOfDeletedAssets = true;
 
-	/** Whether or not to validate assets on save */
-	UPROPERTY(EditAnywhere, config, Category = "Validation Scenarios")
-	uint32 bValidateOnSave : 1;
+	/**
+	 * Minimum severity of validation messages to make the message log visible after validation. 
+	 * Defaults to warning, can be disabled by emptying the optional.
+	 */
+	TOptional<EMessageSeverity::Type> ShowMessageLogSeverity;
+	
+	/** Title of message log page to use for warnings/errors/etc */
+	FText MessageLogPageTitle;
 };
 
 /**
  * UEditorValidatorSubsystem manages all the asset validation in the engine. 
+ * 
  * The first validation handled is UObject::IsDataValid and its overridden functions.
  * Those validations require custom classes and are most suited to project-specific
- * classes. The next validation set is of all registered UEditorValidationBases.
- * These validators have a function to determine if they can validate a given asset,
- * and if they are currently enabled. They are good candidates for validating engine
- * classes or very specific project logic.
+ * classes. 
+ * 
+ * The next validation set is of all registered UEditorValidationBases. These validators 
+ * have a function to determine if they can validate a given asset, and if they are 
+ * currently enabled. They are good candidates for validating engine classes or 
+ * very specific project logic.
+ * 
+ * Finally, this subsystem may be subclassed to change the overally behavior of
+ * validation in your project. If a subclass exist in your project module, it will
+ * supercede the engine validation subsystem.
  */
 UCLASS(Config = Editor)
 class DATAVALIDATION_API UEditorValidatorSubsystem : public UEditorSubsystem
@@ -140,10 +175,44 @@ class DATAVALIDATION_API UEditorValidatorSubsystem : public UEditorSubsystem
 public:
 	UEditorValidatorSubsystem();
 
+	virtual bool ShouldCreateSubsystem(UObject* Outer) const override;
 	virtual void Initialize(FSubsystemCollectionBase& Collection);
-
 	virtual void Deinitialize();
 
+	/**
+	 * Called to validate assets from either the UI or a commandlet.
+	 * Loads the specified assets and runs all registered validators on them.
+	 * Populates the message log with errors and warnings with clickable links.
+	 * @param InSettings Structure passing context and settings for ValidateAssetsWithSettings
+	 * @param OutResults More detailed information about the results of the validate assets command
+	 * @returns Number of assets with validation failures or warnings
+	 */
+	UFUNCTION(BlueprintCallable, Category = "Asset Validation")
+	virtual int32 ValidateAssetsWithSettings(
+		const TArray<FAssetData>& AssetDataList,
+		FValidateAssetsSettings& InSettings,
+		FValidateAssetsResults& OutResults) const;
+
+	/**
+	 * Called to validate assets from either the UI or a commandlet.
+	 * Loads the specified assets and runs all registered validators on them.
+	 * Populates the message log with errors and warnings with clickable links.
+	 * @param InSettings Structure passing context and settings for ValidateAssetsWithSettings
+	 * @param OutResults More detailed information about the results of the validate assets command
+	 * @returns Validation results for the changelist object itself
+	 */
+	UFUNCTION(BlueprintCallable, Category = "Asset Validation")
+	virtual EDataValidationResult ValidateChangelist(
+		UDataValidationChangelist* InChangelist,
+		const FValidateAssetsSettings& InSettings,
+		FValidateAssetsResults& OutResults) const;
+
+	UFUNCTION(BlueprintCallable, Category = "Asset Validation")
+	virtual EDataValidationResult ValidateChangelists(
+		const TArray<UDataValidationChangelist*> InChangelists,
+		const FValidateAssetsSettings& InSettings,
+		FValidateAssetsResults& OutResults) const;
+
 	/*
 	* Adds a validator to the list, making sure it is a unique instance
 	*/
@@ -151,32 +220,46 @@ public:
 	void AddValidator(UEditorValidatorBase* InValidator);
 
 	/**
+	 * Runs registered validators on the provided object.
+	 * Does not add anything to any FMessageLog tabs.
 	 * @return Returns Valid if the object contains valid data; returns Invalid if the object contains invalid data; returns NotValidated if no validations was performed on the object
 	 */
 	UFUNCTION(BlueprintCallable, Category = "Asset Validation")
-	virtual EDataValidationResult IsObjectValid(UObject* InObject, TArray<FText>& ValidationErrors, TArray<FText>& ValidationWarnings, const EDataValidationUsecase InValidationUsecase) const;
+	virtual EDataValidationResult IsObjectValid(
+		UObject* InObject,
+		TArray<FText>& ValidationErrors,
+		TArray<FText>& ValidationWarnings,
+		const EDataValidationUsecase InValidationUsecase) const;
 
 	/**
+	 * Runs registered validators on the provided object.
+	 * Does not add anything to any FMessageLog tabs.
+	 * @return Returns Valid if the object contains valid data; returns Invalid if the object contains invalid data; returns NotValidated if no validations was performed on the object
+	 */
+	virtual EDataValidationResult IsObjectValidWithContext(
+		UObject* InObject,
+		FDataValidationContext& InContext) const;
+
+	/**
+	 * Loads the object referred to by the provided AssetData and runs registered validators on it.
+	 * Does not add anything to any FMessageLog tabs.
 	 * @return Returns Valid if the object pointed to by AssetData contains valid data; returns Invalid if the object contains invalid data or does not exist; returns NotValidated if no validations was performed on the object
 	 */
 	UFUNCTION(BlueprintCallable, Category = "Asset Validation")
-	virtual EDataValidationResult IsAssetValid(const FAssetData& AssetData, TArray<FText>& ValidationErrors, TArray<FText>& ValidationWarnings, const EDataValidationUsecase InValidationUsecase) const;
+	virtual EDataValidationResult IsAssetValid(
+		const FAssetData& AssetData,
+		TArray<FText>& ValidationErrors,
+		TArray<FText>& ValidationWarnings,
+		const EDataValidationUsecase InValidationUsecase) const;
 
 	/**
-	 * Called to validate assets from either the UI or a commandlet
-	 * @param InSettings Structure passing context and settings for ValidateAssetsWithSettings
-	 * @param OutResults More detailed information about the results of the validate assets command
-	 * @returns Number of assets with validation failures or warnings
+	 * Loads the object referred to by the provided AssetData and runs registered validators on it.
+	 * Does not add anything to any FMessageLog tabs.
+	 * @return Returns Valid if the object pointed to by AssetData contains valid data; returns Invalid if the object contains invalid data or does not exist; returns NotValidated if no validations was performed on the object
 	 */
-	UFUNCTION(BlueprintCallable, Category = "Asset Validation")
-	virtual int32 ValidateAssetsWithSettings(const TArray<FAssetData>& AssetDataList, const FValidateAssetsSettings& InSettings, FValidateAssetsResults& OutResults) const;
-	
-	UE_DEPRECATED(5.0, "Use ValidateAssetsWithSettings instead")
-	UFUNCTION(BlueprintCallable, Category = "Asset Validation", meta = (DeprecatedFunction, DeprecationMessage = "Use ValidateAssetsWithSettings instead"))
-	virtual int32 ValidateAssets(TArray<FAssetData> AssetDataList, bool bSkipExcludedDirectories = true, bool bShowIfNoFailures = true) const;
-	
-	UE_DEPRECATED(5.0, "Use version that takes bProceduralSave instead")
-	virtual void ValidateOnSave(TArray<FAssetData> AssetDataList) const;
+	virtual EDataValidationResult IsAssetValidWithContext(
+		const FAssetData& AssetData,
+		FDataValidationContext& InContext) const;
 
 	/**
 	 * Called to validate from an interactive save
@@ -186,9 +269,6 @@ public:
 	 */
 	virtual void ValidateOnSave(TArray<FAssetData> AssetDataList, bool bProceduralSave) const;
 
-	UE_DEPRECATED(5.0, "Use version that takes bProceduralSave instead")
-	virtual void ValidateSavedPackage(FName PackageName);
-
 	/**
 	 * Schedule a validation of a saved package, this will activate next frame by default so it can combine them
 	 *
@@ -196,6 +276,30 @@ public:
 	 * @param bProceduralSave True iff this is e.g. a save from the editor, false if it is a procedural change from e.g. cooking
 	 */
 	virtual void ValidateSavedPackage(FName PackageName, bool bProceduralSave);
+	
+	/** 
+	 * From a changelist, return a list of assets to validate.
+	 * The base implementation returns assets in modified packages in the changelist.
+	 * Subclasses may wish to validate additional assets based on files in the changelist 
+	 * (e.g. dependencies, relevant code changes, configuration changes, etc)
+	 */
+	virtual void GatherAssetsToValidateFromChangelist(
+		UDataValidationChangelist* 		InChangelist,
+		const FValidateAssetsSettings& 	Settings,
+		TSet<FAssetData>& 				OutAssets,
+		FDataValidationContext& 		InContext) const;
+
+	
+	/** 
+	 * Returns true if the given asset should be validated at all.
+	 * May be overridden to e.g. skip developer/test assets.
+	 * Used to determine which assets should be validated when gathering assets from a source such as a changelist.
+	 * Will be ignored when an asset is directly passed for validation.
+	 */
+	virtual bool ShouldValidateAsset(
+		const FAssetData& 				Asset,
+		const FValidateAssetsSettings& 	Settings,
+		FDataValidationContext& 		InContext) const;
 
 protected:
 	void CleanupValidators();
@@ -203,7 +307,7 @@ protected:
 	/**
 	 * @return Returns true if the current Path should be skipped for validation. Returns false otherwise.
 	 */
-	virtual bool IsPathExcludedFromValidation(const FString& Path) const;
+	virtual bool IsPathExcludedFromValidation(FStringView Path) const;
 
 	/**
 	 * Handles validating all pending save packages
@@ -224,6 +328,27 @@ protected:
 	*/
 	void AddValidator(const FAssetData&	InValidatorAssetData);
 
+	/** 
+	 * Validate a set of assets, adding the results to the log/output of a higher level task (e.g. changelist validation).
+	 */
+	void ValidateAssetsInternal(
+		FMessageLog& 					DataValidationLog,
+		TSet<FAssetData>				Assets,
+		const FValidateAssetsSettings& 	InSettings,
+		FValidateAssetsResults& 		OutResults
+	) const;
+	
+	EDataValidationResult ValidateChangelistsInternal(
+		TConstArrayView<UDataValidationChangelist*> InChangelists,
+		const FValidateAssetsSettings& 				InSettings,
+		FValidateAssetsResults& 					OutResults) const;
+	
+	void LogAssetValidationSummary(FMessageLog& DataValidationLog, const FValidateAssetsSettings& InSettings, const FValidateAssetsResults& Results) const;
+
+	EDataValidationResult ValidateObjectInternal(
+		const FAssetData& InAssetData,
+		UObject* InObject,
+		FDataValidationContext& InContext) const;
 protected:
 	/**
 	 * Directories to ignore for data validation. Useful for test assets
diff --git a/Engine/Plugins/Editor/DataValidation/Source/DataValidation/Public/EditorValidator_Localization.h b/Engine/Plugins/Editor/DataValidation/Source/DataValidation/Public/EditorValidator_Localization.h
--- a/Engine/Plugins/Editor/DataValidation/Source/DataValidation/Public/EditorValidator_Localization.h
+++ b/Engine/Plugins/Editor/DataValidation/Source/DataValidation/Public/EditorValidator_Localization.h
@@ -18,8 +18,8 @@ public:
 	UEditorValidator_Localization();
 
 protected:
-	virtual bool CanValidateAsset_Implementation(UObject* InAsset) const override;
-	virtual EDataValidationResult ValidateLoadedAsset_Implementation(UObject* InAsset, TArray<FText>& ValidationErrors) override;
+	virtual bool CanValidateAsset_Implementation(const FAssetData& InAssetData, UObject* InAsset, FDataValidationContext& InContext) const override;
+	virtual EDataValidationResult ValidateLoadedAsset_Implementation(const FAssetData& InAssetData, UObject* InAsset, FDataValidationContext& InContext) override;
 
 	const TArray<FString>* FindOrCacheCulturesForLocalizedRoot(const FString& InLocalizedRootPath);
 
diff --git a/Engine/Plugins/Experimental/Mutable/Source/MutableValidation/Private/MuV/AssetValidator_CustomizableObjects.cpp b/Engine/Plugins/Experimental/Mutable/Source/MutableValidation/Private/MuV/AssetValidator_CustomizableObjects.cpp
--- a/Engine/Plugins/Experimental/Mutable/Source/MutableValidation/Private/MuV/AssetValidator_CustomizableObjects.cpp
+++ b/Engine/Plugins/Experimental/Mutable/Source/MutableValidation/Private/MuV/AssetValidator_CustomizableObjects.cpp
@@ -22,22 +22,19 @@ UAssetValidator_CustomizableObjects::UAssetValidator_CustomizableObjects() : Sup
 	bIsEnabled = true;
 }
 
-
-bool UAssetValidator_CustomizableObjects::CanValidate_Implementation(const EDataValidationUsecase InUsecase) const
+bool UAssetValidator_CustomizableObjects::CanValidateAsset_Implementation(const FAssetData& AssetData, UObject* InAsset, FDataValidationContext& InContext) const
 {
 	// Do not run if saving or running a commandlet (we do not want CIS failing due to our warnings and errors)
-	return !(InUsecase == EDataValidationUsecase::Save || InUsecase == EDataValidationUsecase::Commandlet);
-}
-
+	if (InContext.GetValidationUsecase() == EDataValidationUsecase::Save || InContext.GetValidationUsecase() == EDataValidationUsecase::Commandlet)
+	{
+		return false;
+	}
 
-bool UAssetValidator_CustomizableObjects::CanValidateAsset_Implementation(UObject* InAsset) const
-{
-	return (InAsset ? InAsset->IsA(UCustomizableObject::StaticClass())  : false) ;
+	return Cast<UCustomizableObject>(InAsset) != nullptr;
 }
 
 
-EDataValidationResult UAssetValidator_CustomizableObjects::ValidateLoadedAsset_Implementation(UObject* InAsset,
-	TArray<FText>& ValidationErrors)
+EDataValidationResult UAssetValidator_CustomizableObjects::ValidateLoadedAsset_Implementation(const FAssetData& AssetData, UObject* InAsset, FDataValidationContext& InContext)
 {
 	check(InAsset);
 
@@ -61,11 +58,11 @@ EDataValidationResult UAssetValidator_CustomizableObjects::ValidateLoadedAsset_I
 		// Cache error logs -> They will tag the asset validation as failed
 		for (const FText& ErrorMessage : CoValidationErrors)
 		{
-			AssetFails(InAsset,ErrorMessage,ValidationErrors);
+			AssetFails(InAsset,ErrorMessage);
 		}
 
 		const FText ErrorMessage = FText::Format(LOCTEXT("CO_Validation_Failed", "Validation compilation of {0} CO failed."),  FText::FromString( CustomizableObjectToValidate->GetName()));
-		AssetFails(InAsset,ErrorMessage,ValidationErrors);
+		AssetFails(InAsset,ErrorMessage);
 	}
 	else
 	{
diff --git a/Engine/Plugins/Experimental/Mutable/Source/MutableValidation/Private/MuV/AssetValidator_ReferencedCustomizableObjects.cpp b/Engine/Plugins/Experimental/Mutable/Source/MutableValidation/Private/MuV/AssetValidator_ReferencedCustomizableObjects.cpp
--- a/Engine/Plugins/Experimental/Mutable/Source/MutableValidation/Private/MuV/AssetValidator_ReferencedCustomizableObjects.cpp
+++ b/Engine/Plugins/Experimental/Mutable/Source/MutableValidation/Private/MuV/AssetValidator_ReferencedCustomizableObjects.cpp
@@ -23,24 +23,7 @@ UAssetValidator_ReferencedCustomizableObjects::UAssetValidator_ReferencedCustomi
 	bIsEnabled = true;
 }
 
-
-bool UAssetValidator_ReferencedCustomizableObjects::CanValidate_Implementation(const EDataValidationUsecase InUsecase) const
-{
-	// Use module settings to decide if it needs to run or not.
-    if (const UMutableValidationSettings* ValidationSettings = GetDefault<UMutableValidationSettings>())
-    {
-    	if (!ValidationSettings->bEnableIndirectValidation)
-    	{
-    		return false;
-    	}
-    }
-
-	// Do not run if saving or running a commandlet (we do not want CIS failing due to our warnings and errors)
-	return !(InUsecase == EDataValidationUsecase::Save || InUsecase == EDataValidationUsecase::Commandlet);
-}
-
-
-bool UAssetValidator_ReferencedCustomizableObjects::CanValidateAsset_Implementation(UObject* InAsset) const
+bool UAssetValidator_ReferencedCustomizableObjects::CanValidateAsset_Implementation(const FAssetData& AssetData, UObject* InAsset, FDataValidationContext& InContext) const
 {
 	// Use module settings to decide if it needs to run or not.
 	if (const UMutableValidationSettings* ValidationSettings = GetDefault<UMutableValidationSettings>())
@@ -51,6 +34,12 @@ bool UAssetValidator_ReferencedCustomizableObjects::CanValidateAsset_Implementat
 		}
 	}
 
+	// Do not run if saving or running a commandlet (we do not want CIS failing due to our warnings and errors)
+	if (InContext.GetValidationUsecase() == EDataValidationUsecase::Save || InContext.GetValidationUsecase() == EDataValidationUsecase::Commandlet)
+	{
+		return false;
+	}
+
 	return (InAsset ?
 		InAsset->IsA(UMaterial::StaticClass()) ||
 		InAsset->IsA(UTexture::StaticClass()) ||
@@ -59,8 +48,7 @@ bool UAssetValidator_ReferencedCustomizableObjects::CanValidateAsset_Implementat
 }
 
 
-EDataValidationResult UAssetValidator_ReferencedCustomizableObjects::ValidateLoadedAsset_Implementation(UObject* InAsset,
-	TArray<FText>& ValidationErrors)
+EDataValidationResult UAssetValidator_ReferencedCustomizableObjects::ValidateLoadedAsset_Implementation(const FAssetData& AssetData, UObject* InAsset, FDataValidationContext& InContext)
 {
 	check(InAsset);
 
@@ -74,7 +62,7 @@ EDataValidationResult UAssetValidator_ReferencedCustomizableObjects::ValidateLoa
 	const TSet<UCustomizableObject*> FoundCustomizableObjects = FindCustomizableObjects(FoundReferencers, AssetRegistry);
 
 	// Validate all Customizable Objects we have found
-	ValidateCustomizableObjects(InAsset,FoundCustomizableObjects,ValidationErrors);
+	ValidateCustomizableObjects(InAsset, FoundCustomizableObjects);
 	
 	// Compute InAsset validation status
 	if (GetValidationResult() != EDataValidationResult::Invalid)
@@ -152,7 +140,7 @@ TSet< UCustomizableObject*> UAssetValidator_ReferencedCustomizableObjects::FindC
 }
 
 
-void UAssetValidator_ReferencedCustomizableObjects::ValidateCustomizableObjects(UObject* InAsset, const TSet<UCustomizableObject*>& InCustomizableObjectsToValidate, TArray<FText>& InValidationErrors)
+void UAssetValidator_ReferencedCustomizableObjects::ValidateCustomizableObjects(UObject* InAsset, const TSet<UCustomizableObject*>& InCustomizableObjectsToValidate)
 {
 	for (UCustomizableObject* CustomizableObjectToValidate : InCustomizableObjectsToValidate)
 	{
@@ -173,12 +161,12 @@ void UAssetValidator_ReferencedCustomizableObjects::ValidateCustomizableObjects(
 			// Cache error logs
 			for (const FText& ErrorMessage : CoValidationErrors)
 			{
-				AssetFails(InAsset,ErrorMessage,InValidationErrors);
+				AssetFails(InAsset,ErrorMessage);
 			}
 			
 			// If we say it failed the asset will be marked as bad (containing bad data) and the validator will mark the overall result as failed
 			const FText ErrorMessage = FText::Format(LOCTEXT("RelatedToCustomizableObjectValidator", "The referenced ""\"{0}""\" Mutable Customizable Object is invalid."),  FText::FromString(CustomizableObjectToValidate->GetPathName()));
-			AssetFails(InAsset,ErrorMessage,InValidationErrors);
+			AssetFails(InAsset,ErrorMessage);
 		}
 	}
 }
diff --git a/Engine/Plugins/Experimental/Mutable/Source/MutableValidation/Public/MuV/AssetValidator_CustomizableObjects.h b/Engine/Plugins/Experimental/Mutable/Source/MutableValidation/Public/MuV/AssetValidator_CustomizableObjects.h
--- a/Engine/Plugins/Experimental/Mutable/Source/MutableValidation/Public/MuV/AssetValidator_CustomizableObjects.h
+++ b/Engine/Plugins/Experimental/Mutable/Source/MutableValidation/Public/MuV/AssetValidator_CustomizableObjects.h
@@ -32,9 +32,8 @@ public:
 	
 protected:
 	// UEditorValidatorBase
-	virtual bool CanValidate_Implementation(const EDataValidationUsecase InUsecase) const override;
-	virtual bool CanValidateAsset_Implementation(UObject* InAsset) const override;
-	virtual EDataValidationResult ValidateLoadedAsset_Implementation(UObject* InAsset, TArray<FText>& ValidationErrors) override;
+	virtual bool CanValidateAsset_Implementation(const FAssetData& InAssetData, UObject* InAsset, FDataValidationContext& InContext) const override;
+	virtual EDataValidationResult ValidateLoadedAsset_Implementation(const FAssetData& InAssetData, UObject* InAsset, FDataValidationContext& InContext) override;
 	// UEditorValidatorBase
 
 private:
diff --git a/Engine/Plugins/Experimental/Mutable/Source/MutableValidation/Public/MuV/AssetValidator_ReferencedCustomizableObjects.h b/Engine/Plugins/Experimental/Mutable/Source/MutableValidation/Public/MuV/AssetValidator_ReferencedCustomizableObjects.h
--- a/Engine/Plugins/Experimental/Mutable/Source/MutableValidation/Public/MuV/AssetValidator_ReferencedCustomizableObjects.h
+++ b/Engine/Plugins/Experimental/Mutable/Source/MutableValidation/Public/MuV/AssetValidator_ReferencedCustomizableObjects.h
@@ -24,9 +24,8 @@ public:
 protected:
 
 	// UEditorValidatorBase
-	virtual bool CanValidate_Implementation(const EDataValidationUsecase InUsecase) const override;
-	virtual bool CanValidateAsset_Implementation(UObject* InAsset) const override;
-	virtual EDataValidationResult ValidateLoadedAsset_Implementation(UObject* InAsset, TArray<FText>& ValidationErrors) override;
+	virtual bool CanValidateAsset_Implementation(const FAssetData& InAssetData, UObject* InAsset, FDataValidationContext& InContext) const override;
+	virtual EDataValidationResult ValidateLoadedAsset_Implementation(const FAssetData& InAssetData, UObject* InAsset, FDataValidationContext& InContext) override;
 	// UEditorValidatorBase
 
 private:
@@ -50,6 +49,6 @@ private:
 	 * @param InCustomizableObjectsToValidate Customizable Objects we want to validate with IsDataValid()
 	 * @param InValidationErrors List to fill with the warnings and errors generated during the validation of the Customizable objects
 	 */
-	void ValidateCustomizableObjects(UObject* InAsset, const TSet<UCustomizableObject*>& InCustomizableObjectsToValidate, TArray<FText>& InValidationErrors);
+	void ValidateCustomizableObjects(UObject* InAsset, const TSet<UCustomizableObject*>& InCustomizableObjectsToValidate);
 
 };
diff --git a/Engine/Source/Editor/SourceControlWindows/Private/SSourceControlChangelists.cpp b/Engine/Source/Editor/SourceControlWindows/Private/SSourceControlChangelists.cpp
--- a/Engine/Source/Editor/SourceControlWindows/Private/SSourceControlChangelists.cpp
+++ b/Engine/Source/Editor/SourceControlWindows/Private/SSourceControlChangelists.cpp
@@ -2226,76 +2226,7 @@ static bool GetChangelistValidationResult(FSourceControlChangelistPtr InChangeli
 
 	bool bValidationResult = true;
 
-	if (ValidationDelegate.ExecuteIfBound(InChangelist, ValidationResult, ValidationErrors, ValidationWarnings))
-	{
-		EMessageSeverity::Type MessageSeverity = EMessageSeverity::Info;
-
-		if (ValidationResult == EDataValidationResult::Invalid || ValidationErrors.Num() > 0)
-		{
-			OutValidationTitleText = LOCTEXT("SourceControl.Submit.ChangelistValidationError", "Changelist validation failed!").ToString();
-			bValidationResult = false;
-			MessageSeverity = EMessageSeverity::Error;
-		}
-		else if (ValidationResult == EDataValidationResult::NotValidated || ValidationWarnings.Num() > 0)
-		{
-			OutValidationTitleText = LOCTEXT("SourceControl.Submit.ChangelistValidationWarning", "Changelist validation has warnings!").ToString();
-			MessageSeverity = EMessageSeverity::Warning;
-		}
-		else
-		{
-			OutValidationTitleText = LOCTEXT("SourceControl.Submit.ChangelistValidationSuccess", "Changelist validation successful!").ToString();
-		}
-
-		FMessageLog SourceControlLog("SourceControl");
-		
-		SourceControlLog.Message(MessageSeverity, FText::FromString(*OutValidationTitleText));
-
-		auto AppendInfo = [](const TArray<FText>& Info, const FString& InfoType, FString& OutText)
-		{
-			const int32 MaxNumLinesDisplayed = 5;
-			int32 NumLinesDisplayed = 0;
-
-			if (Info.Num() > 0)
-			{
-				OutText += LINE_TERMINATOR;
-				OutText += FString::Printf(TEXT("Encountered %d %s:"), Info.Num(), *InfoType);
-
-				for (const FText& Line : Info)
-				{
-					if (NumLinesDisplayed >= MaxNumLinesDisplayed)
-					{
-						OutText += LINE_TERMINATOR;
-						OutText += FString::Printf(TEXT("See log for complete list of %s"), *InfoType);
-						break;
-					}
-
-					OutText += LINE_TERMINATOR;
-					OutText += Line.ToString();
-
-					++NumLinesDisplayed;
-				}
-			}
-		};
-
-		auto LogInfo = [&SourceControlLog](const TArray<FText>& Info, const FString& InfoType, const EMessageSeverity::Type LogVerbosity)
-		{
-			if (Info.Num() > 0)
-			{
-				SourceControlLog.Message(LogVerbosity, FText::Format(LOCTEXT("SourceControl.Validation.ErrorEncountered", "Encountered {0} {1}:"), FText::AsNumber(Info.Num()), FText::FromString(*InfoType)));
-
-				for (const FText& Line : Info)
-				{
-					SourceControlLog.Message(LogVerbosity, Line);
-				}
-			}
-		};
-
-		AppendInfo(ValidationErrors, TEXT("errors"), OutValidationErrorsText);
-		AppendInfo(ValidationWarnings, TEXT("warnings"), OutValidationWarningsText);
-
-		LogInfo(ValidationErrors, TEXT("errors"), EMessageSeverity::Error);
-		LogInfo(ValidationWarnings, TEXT("warnings"), EMessageSeverity::Warning);
-	}
+	ValidationDelegate.ExecuteIfBound(InChangelist, ValidationResult, ValidationErrors, ValidationWarnings);
 
 	return bValidationResult;
 }
diff --git a/Engine/Source/Runtime/CoreUObject/Private/Misc/DataValidation.cpp b/Engine/Source/Runtime/CoreUObject/Private/Misc/DataValidation.cpp
new file mode 100644
--- /dev/null
+++ b/Engine/Source/Runtime/CoreUObject/Private/Misc/DataValidation.cpp
@@ -0,0 +1,96 @@
+// Copyright Epic Games, Inc. All Rights Reserved.
+
+#include "Misc/DataValidation.h"
+
+#include "AssetRegistry/AssetData.h"
+#include "AssetRegistry/AssetDataToken.h"
+#include "UObject/UObjectGlobals.h"
+
+EDataValidationResult CombineDataValidationResults(EDataValidationResult Result1, EDataValidationResult Result2)
+{
+	/**
+	 * Anything combined with an Invalid result is Invalid. Any result combined with a NotValidated result is the same result
+	 *
+	 * The combined results should match the following matrix
+	 *
+	 *				|	NotValidated	|	Valid	|	Invalid
+	 * -------------+-------------------+-----------+----------
+	 * NotValidated	|	NotValidated	|	Valid	|	Invalid
+	 * Valid		|	Valid			|	Valid	|	Invalid
+	 * Invalid		|	Invalid			|	Invalid	|	Invalid
+	 *
+	 */
+
+	if (Result1 == EDataValidationResult::Invalid || Result2 == EDataValidationResult::Invalid)
+	{
+		return EDataValidationResult::Invalid;
+	}
+
+	if (Result1 == EDataValidationResult::Valid || Result2 == EDataValidationResult::Valid)
+	{
+		return EDataValidationResult::Valid;
+	}
+
+	return EDataValidationResult::NotValidated;
+}
+	
+void FDataValidationContext::AddMessage(TSharedRef<FTokenizedMessage> InMessage)
+{
+	switch(InMessage->GetSeverity())
+	{
+	case EMessageSeverity::Error:
+		++NumErrors;
+		break;
+	case EMessageSeverity::Warning:
+	case EMessageSeverity::PerformanceWarning:
+		++NumWarnings;
+		break;
+	case EMessageSeverity::Info:
+	default:
+		break;
+	}	
+	Issues.Emplace(InMessage);
+}
+
+TSharedRef<FTokenizedMessage> FDataValidationContext::AddMessage(const FAssetData& ForAsset, EMessageSeverity::Type InSeverity, FText InText)
+{
+	TSharedRef<FTokenizedMessage> Message = FTokenizedMessage::Create(InSeverity);
+	if (ForAsset.IsValid())
+	{
+		Message->AddToken(FAssetDataToken::Create(ForAsset));
+	}
+	if (!InText.IsEmpty())
+	{
+		Message->AddText(InText);
+	}
+	AddMessage(Message);
+	return Message;
+}
+
+TSharedRef<FTokenizedMessage> FDataValidationContext::AddMessage(EMessageSeverity::Type InSeverity, FText InText)
+{
+	TSharedRef<FTokenizedMessage> Message = FTokenizedMessage::Create(InSeverity, InText);
+	AddMessage(Message);
+	return Message;
+}
+
+void FDataValidationContext::SplitIssues(TArray<FText>& Warnings, TArray<FText>& Errors) const
+{
+	for (const FIssue& Issue : GetIssues())
+	{
+		FText Message = Issue.Message;
+		if (Issue.TokenizedMessage.IsValid())
+		{
+			Message = Issue.TokenizedMessage->ToText();
+		}
+
+		if (Issue.Severity == EMessageSeverity::Error)
+		{
+			Errors.Add(Issue.Message);
+		}
+		else if (Issue.Severity == EMessageSeverity::Warning)
+		{	
+			Warnings.Add(Issue.Message);
+		}
+	}
+}
diff --git a/Engine/Source/Runtime/CoreUObject/Private/UObject/UObjectGlobals.cpp b/Engine/Source/Runtime/CoreUObject/Private/UObject/UObjectGlobals.cpp
--- a/Engine/Source/Runtime/CoreUObject/Private/UObject/UObjectGlobals.cpp
+++ b/Engine/Source/Runtime/CoreUObject/Private/UObject/UObjectGlobals.cpp
@@ -6324,34 +6324,6 @@ namespace UECodeGen_Private
 	}
 }
 
-EDataValidationResult CombineDataValidationResults(EDataValidationResult Result1, EDataValidationResult Result2)
-{
-	/**
-	 * Anything combined with an Invalid result is Invalid. Any result combined with a NotValidated result is the same result
-	 *
-	 * The combined results should match the following matrix
-	 *
-	 *				|	NotValidated	|	Valid	|	Invalid
-	 * -------------+-------------------+-----------+----------
-	 * NotValidated	|	NotValidated	|	Valid	|	Invalid
-	 * Valid		|	Valid			|	Valid	|	Invalid
-	 * Invalid		|	Invalid			|	Invalid	|	Invalid
-	 *
-	 */
-
-	if (Result1 == EDataValidationResult::Invalid || Result2 == EDataValidationResult::Invalid)
-	{
-		return EDataValidationResult::Invalid;
-	}
-
-	if (Result1 == EDataValidationResult::Valid || Result2 == EDataValidationResult::Valid)
-	{
-		return EDataValidationResult::Valid;
-	}
-
-	return EDataValidationResult::NotValidated;
-}
-
 void FReferenceCollector::AddStableReferenceSetFwd(TSet<FObjectPtr>* Objects)
 {
 	AddStableReferenceSet(reinterpret_cast<TSet<TObjectPtr<UObject>>*>(Objects));
diff --git a/Engine/Source/Runtime/CoreUObject/Public/Misc/DataValidation.h b/Engine/Source/Runtime/CoreUObject/Public/Misc/DataValidation.h
--- a/Engine/Source/Runtime/CoreUObject/Public/Misc/DataValidation.h
+++ b/Engine/Source/Runtime/CoreUObject/Public/Misc/DataValidation.h
@@ -5,6 +5,35 @@
 #include "CoreMinimal.h"
 #include "Logging/TokenizedMessage.h"
 
+#include "DataValidation.generated.h"
+
+struct FAssetData;
+
+UENUM(BlueprintType)
+enum class EDataValidationUsecase : uint8
+{
+	/** No usecase specified */
+	None = 0,
+
+	/** Triggered on user's demand */
+	Manual,
+
+	/** A commandlet invoked the validation */
+	Commandlet,
+
+	/** Saving a package triggered the validation */
+	Save,
+
+	/** Submit dialog triggered the validation */
+	PreSubmit,
+
+	/** Triggered by blueprint or c++ */
+	Script,
+};
+
+/** 
+ * This class is used to interface between UObject::IsDataValid and the data validation system.
+ */
 class FDataValidationContext : public FNoncopyable
 {
 public:
@@ -12,13 +41,58 @@ public:
 	{
 		FText Message;
 		EMessageSeverity::Type Severity;
+		TSharedPtr<FTokenizedMessage> TokenizedMessage;
 
+		FIssue(TSharedRef<FTokenizedMessage> InTokenizedMessage)
+			: Message()
+			, Severity(InTokenizedMessage->GetSeverity())
+			, TokenizedMessage(InTokenizedMessage)
+		{}
 		FIssue(const FText& InMessage, EMessageSeverity::Type InSeverity)
 			: Message(InMessage)
 			, Severity(InSeverity)
+			, TokenizedMessage()
 		{}
 	};
+	
+	FDataValidationContext() = default;
+	
+	/** 
+	 * Constructor for use by UEditorValidatorSubsystem.
+	 * @param Inusecase allows validators to skip certain validation steps depending on the validation trigger condition
+	 * @param InAssociatedObjects allows validators to check additional associated objects of the asset being validated e.g. validating external actors
+	 */
+	FDataValidationContext(
+		bool InWasAssetLoadedForValidation,
+		EDataValidationUsecase InUsecase,
+		TConstArrayView<FAssetData> InAssociatedObjects)
+		: ValidationUsecase(InUsecase)
+		, bWasAssetLoadedForValidation(InWasAssetLoadedForValidation)
+	{}
 
+	EDataValidationUsecase GetValidationUsecase() const
+	{
+		return ValidationUsecase;
+	}
+	
+	TConstArrayView<FAssetData> GetAssociatedExternalObjects() const
+	{
+		return AssociatedExternalObjects;
+	}
+	
+	// Returns whether the asset currently being validation was loaded specifically for the purposes of validation
+	// Otherwise it was already loaded before validation started
+	// Default valid is false so that validators can attempt to reload asset to verify on-disk state/load process 
+	// if unsure.
+	// If true, allows that duplicate work to be skipped in some context. 
+	bool WasAssetLoadedForValidation() const
+	{
+		return bWasAssetLoadedForValidation;
+	}
+
+	COREUOBJECT_API TSharedRef<FTokenizedMessage> AddMessage(const FAssetData& ForAsset, EMessageSeverity::Type InSeverity, FText InText = {});
+	COREUOBJECT_API TSharedRef<FTokenizedMessage> AddMessage(EMessageSeverity::Type InSeverity, FText InText = {});
+	COREUOBJECT_API void AddMessage(TSharedRef<FTokenizedMessage> Message);
 	void AddWarning(const FText& Text) { Issues.Emplace(Text, EMessageSeverity::Warning); NumWarnings++; }
 	void AddError(const FText& Text) { Issues.Emplace(Text, EMessageSeverity::Error); NumErrors++; }
 
@@ -26,23 +100,14 @@ public:
 	uint32 GetNumWarnings() const { return NumWarnings; }
 	uint32 GetNumErrors() const { return NumErrors; }
 
-	void SplitIssues(TArray<FText>& Warnings, TArray<FText>& Errors) const
-	{
-		for (const FIssue& Issue : GetIssues())
-		{
-			if (Issue.Severity == EMessageSeverity::Error)
-			{
-				Errors.Add(Issue.Message);
-			}
-			else if (Issue.Severity == EMessageSeverity::Warning)
-			{
-				Warnings.Add(Issue.Message);
-			}
-		}
-	}
-
+	COREUOBJECT_API void SplitIssues(TArray<FText>& Warnings, TArray<FText>& Errors) const;
+	
 private:
 	TArray<FIssue> Issues;
 	uint32 NumWarnings = 0;
 	uint32 NumErrors = 0;
+	
+	TConstArrayView<FAssetData> AssociatedExternalObjects;
+	EDataValidationUsecase ValidationUsecase = EDataValidationUsecase::None;
+	bool bWasAssetLoadedForValidation = false;
 };
\ No newline at end of file
