Disable dirty files validator in commandlet.

diff --git a/Engine/Plugins/Editor/DataValidation/Source/DataValidation/Private/DirtyFilesChangelistValidator.cpp b/Engine/Plugins/Editor/DataValidation/Source/DataValidation/Private/DirtyFilesChangelistValidator.cpp
--- a/Engine/Plugins/Editor/DataValidation/Source/DataValidation/Private/DirtyFilesChangelistValidator.cpp
+++ b/Engine/Plugins/Editor/DataValidation/Source/DataValidation/Private/DirtyFilesChangelistValidator.cpp
@@ -8,6 +8,7 @@
 #include "DataValidationChangelist.h"
 #include "FileHelpers.h"
 #include "ISourceControlModule.h"
+#include "Misc/DataValidation.h"
 #include "Misc/Paths.h"
 #include "Containers/Map.h"
 
@@ -17,6 +18,10 @@
 
 bool UDirtyFilesChangelistValidator::CanValidateAsset_Implementation(const FAssetData& AssetData, UObject* InAsset, FDataValidationContext& InContext) const
 {
+	if (InContext.GetValidationUsecase() == EDataValidationUsecase::Commandlet)
+	{
+		return false;
+	}
 	return (InAsset != nullptr) && (UDataValidationChangelist::StaticClass() == InAsset->GetClass());
 }
 
