Fix duplicate localization keys.

diff --git a/Engine/Plugins/Editor/DataValidation/Source/DataValidation/Private/EditorValidatorSubsystem.cpp b/Engine/Plugins/Editor/DataValidation/Source/DataValidation/Private/EditorValidatorSubsystem.cpp
--- a/Engine/Plugins/Editor/DataValidation/Source/DataValidation/Private/EditorValidatorSubsystem.cpp
+++ b/Engine/Plugins/Editor/DataValidation/Source/DataValidation/Private/EditorValidatorSubsystem.cpp
@@ -533,7 +533,7 @@ void UEditorValidatorSubsystem::ValidateAssetsInternal(
 					else 
 					{
 						DataValidationLog.Info()->AddToken(FAssetDataToken::Create(Data))
-							->AddToken(FTextToken::Create(LOCTEXT("DataValidation.ContainsWarningsResult", "contains valid data.")));
+							->AddToken(FTextToken::Create(LOCTEXT("DataValidation.ValidResult", "contains valid data.")));
 					}
 					break;
 				case EDataValidationResult::NotValidated:
