Add an interface function to editor validators to allow adding additional files to validate from a changelist.

diff --git a/Engine/Plugins/Editor/DataValidation/Source/DataValidation/Private/EditorValidatorSubsystem.cpp b/Engine/Plugins/Editor/DataValidation/Source/DataValidation/Private/EditorValidatorSubsystem.cpp
--- a/Engine/Plugins/Editor/DataValidation/Source/DataValidation/Private/EditorValidatorSubsystem.cpp
+++ b/Engine/Plugins/Editor/DataValidation/Source/DataValidation/Private/EditorValidatorSubsystem.cpp
@@ -852,7 +852,9 @@ void UEditorValidatorSubsystem::GatherAssetsToValidateFromChangelist(
 {
 	FAssetRegistryModule& AssetRegistryModule = FModuleManager::LoadModuleChecked<FAssetRegistryModule>("AssetRegistry");
 	IAssetRegistry& AssetRegistry = AssetRegistryModule.Get();
-
+	
+	LoadValidators();
+	
 	for (const FName& PackageName : InChangelist->ModifiedPackageNames)
 	{
 		TArray<FAssetData> NewAssets;
@@ -860,6 +862,26 @@ void UEditorValidatorSubsystem::GatherAssetsToValidateFromChangelist(
 		OutAssets.Append(NewAssets);
 	}
 	
+	// Gather assets requested by plugin/project validators 
+	for (const TPair<FTopLevelAssetPath, TObjectPtr<UEditorValidatorBase>>& ValidatorPair : Validators)
+	{
+		UEditorValidatorBase* Validator = ValidatorPair.Value;
+		if (Validator && Validator->IsEnabled())
+		{
+			TArray<FAssetData> NewAssets = Validator->GetAssetsToValidateFromChangelist(InChangelist, InContext);
+			for (const FAssetData& Asset : NewAssets)
+			{
+				// It's not strictly necessary to filter assets here but it makes logging simpler
+				if (ShouldValidateAsset(Asset, Settings, InContext))
+				{
+					UE_LOG(LogContentValidation, Verbose, TEXT("Asset validator %s adding %s to be validated."), *Validator->GetPathName(), *Asset.GetSoftObjectPath().ToString());
+					OutAssets.Add(Asset);
+				}
+			}
+		}
+	}
+
+	
 	if (Settings.bValidateReferencersOfDeletedAssets)
 	{
 		TSet<FAssetData> AdditionalAssetsToValidate;
@@ -881,7 +903,12 @@ void UEditorValidatorSubsystem::GatherAssetsToValidateFromChangelist(
 	}
 }
 
-void UEditorValidatorSubsystem::LoadValidators()
+void UEditorValidatorSubsystem::LoadValidators() const
+{
+	const_cast<UEditorValidatorSubsystem*>(this)->LoadValidators();
+}
+
+void UEditorValidatorSubsystem::LoadValidators() 
 {
 	if (bNeedLoadingOfValidators)
 	{
@@ -914,5 +941,38 @@ void UEditorValidatorSubsystem::LoadValidators()
 	}
 }
 
+TArray<FAssetData> UEditorValidatorSubsystem::GetAssetsResolvingRedirectors(FARFilter& InFilter)
+{
+	FAssetRegistryModule& AssetRegistryModule = FModuleManager::LoadModuleChecked<FAssetRegistryModule>("AssetRegistry");
+	IAssetRegistry& AssetRegistry = AssetRegistryModule.Get();
+
+	TArray<FAssetData> Found;
+	AssetRegistry.GetAssets(InFilter, Found);
+	
+	TArray<FAssetData> Redirectors;	
+	for (int32 i=Found.Num()-1; i >= 0; --i)
+	{
+		if (Found[i].IsRedirector())
+		{
+			Redirectors.Add(Found[i]);
+			Found.RemoveAt(i);
+		}
+	}
+	
+	for (const FAssetData& RedirectorAsset : Redirectors)
+	{
+		FSoftObjectPath Path = AssetRegistry.GetRedirectedObjectPath(RedirectorAsset.GetSoftObjectPath());		
+		if (!Path.IsNull())
+		{
+			FAssetData Destination = AssetRegistry.GetAssetByObjectPath(Path, true);
+			if (Destination.IsValid())
+			{
+				Found.Add(Destination);
+			}
+		}
+	}
+	return Found;
+}
+
 #undef LOCTEXT_NAMESPACE
 
diff --git a/Engine/Plugins/Editor/DataValidation/Source/DataValidation/Public/EditorValidatorBase.h b/Engine/Plugins/Editor/DataValidation/Source/DataValidation/Public/EditorValidatorBase.h
--- a/Engine/Plugins/Editor/DataValidation/Source/DataValidation/Public/EditorValidatorBase.h
+++ b/Engine/Plugins/Editor/DataValidation/Source/DataValidation/Public/EditorValidatorBase.h
@@ -11,6 +11,7 @@ namespace EMessageSeverity { enum Type : int; };
 struct FAssetData;
 class FDataValidationContext;
 class FTokenizedMessage;
+class UDataValidationChangelist;
 
 /*
 * The EditorValidatorBase is a class which verifies that an asset meets a specific ruleset.
@@ -48,6 +49,12 @@ public:
 		return EDataValidationResult::NotValidated;
 	}
 	
+	/** Override this to return additional assets that need validation because of other changes in a changelist. e.g. referencers of a specific type may be broken by edits to a dependency. */
+	virtual TArray<FAssetData> GetAssetsToValidateFromChangelist(UDataValidationChangelist* InChangelist, FDataValidationContext& InContext) const
+	{
+		return {};
+	}
+	
 	/** Override this to determine whether or not you can use this validator given this usecase */
 	UFUNCTION(BlueprintNativeEvent, Category = "Asset Validation", DisplayName="Can Validate")
 	bool K2_CanValidate(const EDataValidationUsecase InUsecase) const;
diff --git a/Engine/Plugins/Editor/DataValidation/Source/DataValidation/Public/EditorValidatorSubsystem.h b/Engine/Plugins/Editor/DataValidation/Source/DataValidation/Public/EditorValidatorSubsystem.h
--- a/Engine/Plugins/Editor/DataValidation/Source/DataValidation/Public/EditorValidatorSubsystem.h
+++ b/Engine/Plugins/Editor/DataValidation/Source/DataValidation/Public/EditorValidatorSubsystem.h
@@ -14,6 +14,7 @@ class UEditorValidatorBase;
 class UDataValidationChangelist;
 struct FDirectoryPath;
 class FMessageLog;
+struct FARFilter;
 struct FAssetData;
 class ISourceControlChangelist;
 typedef TSharedPtr<class ISourceControlChangelist, ESPMode::ThreadSafe> FSourceControlChangelistPtr;
@@ -301,6 +302,13 @@ public:
 		const FValidateAssetsSettings& 	Settings,
 		FDataValidationContext& 		InContext) const;
 
+	
+	// Utility functions
+	
+	// Retrieve all assets matching the given filter from the asset registry and then recursively resolve redirectors to produce a single set of assets.
+	TArray<FAssetData> GetAssetsResolvingRedirectors(FARFilter& InFilter);
+
+
 protected:
 	void CleanupValidators();
 
@@ -321,6 +329,7 @@ protected:
 	 */
 	void ValidateChangelistPreSubmit(FSourceControlChangelistPtr Changelist, EDataValidationResult& OutResult, TArray<FText>& ValidationErrors, TArray<FText>& ValidationWarnings) const;
 
+	void LoadValidators() const;
 	void LoadValidators();
 
 	/*
