Fix not returning asset validation result through ValidateChangelists.
Don't wait for asset compilation when validating on save.

diff --git a/Engine/Plugins/Editor/DataValidation/Source/DataValidation/Private/EditorValidatorSubsystem.cpp b/Engine/Plugins/Editor/DataValidation/Source/DataValidation/Private/EditorValidatorSubsystem.cpp
--- a/Engine/Plugins/Editor/DataValidation/Source/DataValidation/Private/EditorValidatorSubsystem.cpp
+++ b/Engine/Plugins/Editor/DataValidation/Source/DataValidation/Private/EditorValidatorSubsystem.cpp
@@ -330,7 +330,7 @@ int32 UEditorValidatorSubsystem::ValidateAssetsWithSettings(
 	return OutResults.NumWarnings + OutResults.NumInvalid;
 }
 
-void UEditorValidatorSubsystem::ValidateAssetsInternal(
+EDataValidationResult UEditorValidatorSubsystem::ValidateAssetsInternal(
 	FMessageLog& 					DataValidationLog,
 	TSet<FAssetData>				AssetDataList,
 	const FValidateAssetsSettings& 	InSettings,
@@ -342,6 +342,14 @@ void UEditorValidatorSubsystem::ValidateAssetsInternal(
 	FScopedSlowTask SlowTask(AssetDataList.Num(), LOCTEXT("DataValidation.ValidateAssetsTask", "Validating Assets"));
 	SlowTask.MakeDialog();
 	
+	LoadValidators();
+	UE_LOG(LogContentValidation, Log, TEXT("Starting to validate %d assets"), AssetDataList.Num());
+	UE_LOG(LogContentValidation, Log, TEXT("Registered validators:"));
+	for (const TPair<FTopLevelAssetPath, TObjectPtr<UEditorValidatorBase>>& ValidatorPair : Validators)
+	{
+		UE_LOG(LogContentValidation, Log, TEXT("\t%s"), *ValidatorPair.Key.ToString());
+	}
+	
 	// Broadcast the Editor event before we start validating. This lets other systems (such as Sequencer) restore the state
 	// of the level to what is actually saved on disk before performing validation.
 	if (FEditorDelegates::OnPreAssetValidation.IsBound())
@@ -377,14 +385,11 @@ void UEditorValidatorSubsystem::ValidateAssetsInternal(
 	}
 
 	// Dont let other async compilation warnings be attributed incorrectly to the package that is loading.
-	if (FAssetCompilingManager::Get().GetNumRemainingAssets())
-	{
-		FScopedSlowTask CompileAssetsSlowTask(0.f, LOCTEXT("DataValidation.CompilingAssetsBeforeCheckingContentTask", "Finishing asset compilations before checking content..."));
-		CompileAssetsSlowTask.MakeDialog();
-		FAssetCompilingManager::Get().FinishAllCompilation();
-	}
+	WaitForAssetCompilationIfNecessary(InSettings.ValidationUsecase);
 
 	OutResults.NumRequested = AssetDataList.Num();
+	
+	EDataValidationResult Result = EDataValidationResult::NotValidated;
 
 	// Now add to map or update as needed
 	for (const FAssetData& Data : AssetDataList)
@@ -429,19 +434,14 @@ void UEditorValidatorSubsystem::ValidateAssetsInternal(
 		const bool bAlreadyLoaded = LoadedAsset != nullptr;
 
 		FDataValidationContext ValidationContext(bAlreadyLoaded, InSettings.ValidationUsecase, AssetsToExternalObjects.FindRef(Data));
-		EDataValidationResult Result = EDataValidationResult::NotValidated;
+		EDataValidationResult AssetResult = EDataValidationResult::NotValidated;
 		if (!LoadedAsset)
 		{
 			UE_LOG(LogContentValidation, Log, TEXT("Loading asset %s for validation"), *Data.ToSoftObjectPath().ToString());
 			UE::DataValidation::FScopedLogMessageGatherer LogGatherer(InSettings.bCaptureAssetLoadLogs);
 			LoadedAsset = Data.GetAsset(); // Do not load external objects, validators should load external objects that they want via GetAssociatedExternalObjects in the validation context 
-			
-			if (FAssetCompilingManager::Get().GetNumRemainingAssets())
-			{
-				FScopedSlowTask CompileAssetsSlowTask(0.f, LOCTEXT("DataValidation.CompilingAssetsBeforeCheckingContentTask", "Finishing asset compilations before checking content..."));
-				CompileAssetsSlowTask.MakeDialog();
-				FAssetCompilingManager::Get().FinishAllCompilation();
-			}
+
+			WaitForAssetCompilationIfNecessary(InSettings.ValidationUsecase);
 			
 			// Associate any load errors with this asset in the message log
 			TArray<FString> Warnings;
@@ -462,14 +462,14 @@ void UEditorValidatorSubsystem::ValidateAssetsInternal(
 				ValidationContext.AddMessage(EMessageSeverity::Error)
 					->AddToken(FAssetDataToken::Create(Data))
 					->AddText(LOCTEXT("DataValidation.LoadErrors", "Errors loading asset {0}"), FText::FromStringView(Buffer.ToView()));
-				Result = EDataValidationResult::Invalid;
+				AssetResult = EDataValidationResult::Invalid;
 			}
 		}
 
 		if (LoadedAsset)
 		{
 			UE::DataValidation::FScopedLogMessageGatherer LogGatherer(InSettings.bCaptureLogsDuringValidation);
-			Result = IsObjectValidWithContext(LoadedAsset, ValidationContext);
+			AssetResult = IsObjectValidWithContext(LoadedAsset, ValidationContext);
 			
 			// Associate any UE_LOG errors with this asset in the message log
 			TArray<FString> Warnings;
@@ -482,7 +482,7 @@ void UEditorValidatorSubsystem::ValidateAssetsInternal(
 				ValidationContext.AddMessage(EMessageSeverity::Error)
 					->AddToken(FAssetDataToken::Create(Data))
 					->AddText(LOCTEXT("DataValidation.DuringValidationWarnings", "Warnings logged while validating asset {0}"), FText::FromStringView(Buffer.ToView()));
-				Result = EDataValidationResult::Invalid;
+				AssetResult = EDataValidationResult::Invalid;
 			}
 			if(Errors.Num())
 			{
@@ -491,7 +491,7 @@ void UEditorValidatorSubsystem::ValidateAssetsInternal(
 				ValidationContext.AddMessage(EMessageSeverity::Error)
 					->AddToken(FAssetDataToken::Create(Data))
 					->AddText(LOCTEXT("DataValidation.DuringValidationErrors", "Errors logged while validating asset {0}"), FText::FromStringView(Buffer.ToView()));
-				Result = EDataValidationResult::Invalid;
+				AssetResult = EDataValidationResult::Invalid;
 			}
 		}
 		else if (InSettings.bLoadAssetsForValidation)
@@ -499,14 +499,14 @@ void UEditorValidatorSubsystem::ValidateAssetsInternal(
 			ValidationContext.AddMessage(EMessageSeverity::Error)
 				->AddToken(FAssetDataToken::Create(Data))
 				->AddToken(FTextToken::Create(LOCTEXT("DataValidation.LoadFailed", "Failed to load object")));
-			Result = EDataValidationResult::Invalid;
+			AssetResult = EDataValidationResult::Invalid;
 		}
 		else 
 		{
 			ValidationContext.AddMessage(EMessageSeverity::Error)
 				->AddToken(FAssetDataToken::Create(Data))
 				->AddToken(FTextToken::Create(LOCTEXT("DataValidation.CannotValidateNotLoaded", "Cannot validate unloaded asset")));
-			Result = EDataValidationResult::Invalid;
+			AssetResult = EDataValidationResult::Invalid;
 		}
 
 		++OutResults.NumChecked;
@@ -522,7 +522,7 @@ void UEditorValidatorSubsystem::ValidateAssetsInternal(
 
 		if (InSettings.bShowIfNoFailures)
 		{
-			switch (Result)
+			switch (AssetResult)
 			{
 				case EDataValidationResult::Valid:
 					if (bAnyWarnings)
@@ -543,7 +543,7 @@ void UEditorValidatorSubsystem::ValidateAssetsInternal(
 			}
 		}	
 
-		switch (Result)
+		switch (AssetResult)
 		{
 			case EDataValidationResult::Valid:
 				++OutResults.NumValid;
@@ -561,11 +561,13 @@ void UEditorValidatorSubsystem::ValidateAssetsInternal(
 			FValidateAssetsDetails& Details = OutResults.AssetsDetails.Emplace(Data.GetObjectPathString());
 			Details.PackageName = Data.PackageName;
 			Details.AssetName = Data.AssetName;
-			Details.Result = Result;
+			Details.Result = AssetResult;
 			ValidationContext.SplitIssues(Details.ValidationErrors, Details.ValidationWarnings);
 		}
 		
 		DataValidationLog.Flush();
+		
+		Result = CombineDataValidationResults(Result, AssetResult);
 	}
 
 	// Broadcast now that we're complete so other systems can go back to their previous state.
@@ -573,6 +575,8 @@ void UEditorValidatorSubsystem::ValidateAssetsInternal(
 	{
 		FEditorDelegates::OnPostAssetValidation.Broadcast();
 	}
+	
+	return Result;
 }
 
 void UEditorValidatorSubsystem::LogAssetValidationSummary(FMessageLog& DataValidationLog, const FValidateAssetsSettings& InSettings, const FValidateAssetsResults& Results) const
@@ -838,7 +842,8 @@ EDataValidationResult UEditorValidatorSubsystem::ValidateChangelistsInternal(
 	}
 
 	// Validate assets from all changelists
-	ValidateAssetsInternal(DataValidationLog, MoveTemp(AssetsToValidate), Settings, OutResults);
+	EDataValidationResult AssetResult = ValidateAssetsInternal(DataValidationLog, MoveTemp(AssetsToValidate), Settings, OutResults);
+	Result = CombineDataValidationResults(Result, AssetResult);
 	LogAssetValidationSummary(DataValidationLog, Settings, OutResults);
 
 	return Result;
@@ -974,5 +979,20 @@ TArray<FAssetData> UEditorValidatorSubsystem::GetAssetsResolvingRedirectors(FARF
 	return Found;
 }
 
+void UEditorValidatorSubsystem::WaitForAssetCompilationIfNecessary(EDataValidationUsecase InUsecase) const
+{
+	if (InUsecase == EDataValidationUsecase::Save)
+	{
+		return;
+	}
+
+	if (FAssetCompilingManager::Get().GetNumRemainingAssets())
+	{
+		FScopedSlowTask CompileAssetsSlowTask(0.f, LOCTEXT("DataValidation.CompilingAssetsBeforeCheckingContentTask", "Finishing asset compilations before checking content..."));
+		CompileAssetsSlowTask.MakeDialog();
+		FAssetCompilingManager::Get().FinishAllCompilation();
+	}
+}
+
 #undef LOCTEXT_NAMESPACE
 
diff --git a/Engine/Plugins/Editor/DataValidation/Source/DataValidation/Public/EditorValidatorSubsystem.h b/Engine/Plugins/Editor/DataValidation/Source/DataValidation/Public/EditorValidatorSubsystem.h
--- a/Engine/Plugins/Editor/DataValidation/Source/DataValidation/Public/EditorValidatorSubsystem.h
+++ b/Engine/Plugins/Editor/DataValidation/Source/DataValidation/Public/EditorValidatorSubsystem.h
@@ -312,6 +312,8 @@ public:
 
 protected:
 	void CleanupValidators();
+	
+	void WaitForAssetCompilationIfNecessary(EDataValidationUsecase InUsecase) const;
 
 	/**
 	 * @return Returns true if the current Path should be skipped for validation. Returns false otherwise.
@@ -341,7 +343,7 @@ protected:
 	/** 
 	 * Validate a set of assets, adding the results to the log/output of a higher level task (e.g. changelist validation).
 	 */
-	void ValidateAssetsInternal(
+	EDataValidationResult ValidateAssetsInternal(
 		FMessageLog& 					DataValidationLog,
 		TSet<FAssetData>				Assets,
 		const FValidateAssetsSettings& 	InSettings,
