Fix use of BP/Python validators

diff --git a/Engine/Plugins/Editor/DataValidation/Source/DataValidation/Private/EditorValidatorBase.cpp b/Engine/Plugins/Editor/DataValidation/Source/DataValidation/Private/EditorValidatorBase.cpp
--- a/Engine/Plugins/Editor/DataValidation/Source/DataValidation/Private/EditorValidatorBase.cpp
+++ b/Engine/Plugins/Editor/DataValidation/Source/DataValidation/Private/EditorValidatorBase.cpp
@@ -27,7 +27,7 @@ EDataValidationResult UEditorValidatorBase::ValidateLoadedAsset(const FAssetData
 	TGuardValue<UObject*> ObjectGuard(CurrentObjectBeingValidated, Asset);
 	TGuardValue<const FAssetData*> AssetGuard(CurrentAssetBeingValidated, &AssetData);
 
-	if (K2_CanValidate(Context.GetValidationUsecase()))
+	if (K2_CanValidate(Context.GetValidationUsecase()) && K2_CanValidateAsset(Asset))
 	{
 		EDataValidationResult K2Result = K2_ValidateLoadedAsset(Asset);
 		K2Result = CombineDataValidationResults(GetValidationResult(), K2Result);
@@ -65,7 +65,7 @@ PRAGMA_ENABLE_DEPRECATION_WARNINGS
 
 bool UEditorValidatorBase::K2_CanValidate_Implementation(const EDataValidationUsecase InUsecase) const
 {
-	return false;
+	return true;
 }
 
 bool UEditorValidatorBase::K2_CanValidateAsset_Implementation(UObject* InAsset) const
