Move elapsed time message to very verbose and move it to ValidateLoadedAsset

diff --git a/Engine/Plugins/Editor/DataValidation/Source/DataValidation/Private/EditorValidatorBase.cpp b/Engine/Plugins/Editor/DataValidation/Source/DataValidation/Private/EditorValidatorBase.cpp
--- a/Engine/Plugins/Editor/DataValidation/Source/DataValidation/Private/EditorValidatorBase.cpp
+++ b/Engine/Plugins/Editor/DataValidation/Source/DataValidation/Private/EditorValidatorBase.cpp
@@ -23,6 +23,8 @@ EDataValidationResult UEditorValidatorBase::ValidateLoadedAsset(const FAssetData
 	EDataValidationResult Result = EDataValidationResult::NotValidated;
 	
 	ResetValidationState();
+	FDateTime ValidationTime = FDateTime::Now();
+	bool bTriedToValidate = false;
 	
 	TGuardValue<UObject*> ObjectGuard(CurrentObjectBeingValidated, Asset);
 	TGuardValue<const FAssetData*> AssetGuard(CurrentAssetBeingValidated, &AssetData);
@@ -32,6 +34,7 @@ EDataValidationResult UEditorValidatorBase::ValidateLoadedAsset(const FAssetData
 		EDataValidationResult K2Result = K2_ValidateLoadedAsset(Asset);
 		K2Result = CombineDataValidationResults(GetValidationResult(), K2Result);
 		ensureMsgf(K2Result != EDataValidationResult::NotValidated, TEXT("Validator %s did not return a validation result from BP validation path for asset %s"), *GetClass()->GetPathName(), *Asset->GetPathName());
+		bTriedToValidate = true;
 	}
 	
 PRAGMA_DISABLE_DEPRECATION_WARNINGS
@@ -45,6 +48,7 @@ PRAGMA_DISABLE_DEPRECATION_WARNINGS
 		ensureMsgf(NewResult != EDataValidationResult::NotValidated, TEXT("Validator %s did not return a validation result from legacy validation path for asset %s"), *GetClass()->GetPathName(), *Asset->GetPathName());
 
 		Result = CombineDataValidationResults(Result, NewResult);
+		bTriedToValidate = true;
 	}
 PRAGMA_ENABLE_DEPRECATION_WARNINGS
 	
@@ -56,8 +60,22 @@ PRAGMA_ENABLE_DEPRECATION_WARNINGS
 		ensureMsgf(NewResult != EDataValidationResult::NotValidated, TEXT("Validator %s did not return a validation result from native validation path for asset %s"), *GetClass()->GetPathName(), *Asset->GetPathName());
 
 		Result = CombineDataValidationResults(Result, NewResult);
+		bTriedToValidate = true;
 	}
 	
+	if(bTriedToValidate && LogContentValidation.GetVerbosity() >= ELogVerbosity::VeryVerbose)
+	{
+		FDateTime CurrentTime = FDateTime::Now();
+		FTimespan ElapsedTimeSpan = (CurrentTime - ValidationTime);
+		float ElapsedTimeMS = ElapsedTimeSpan.GetTotalMilliseconds();
+		FNumberFormattingOptions TimeFormat;
+		TimeFormat.MinimumFractionalDigits = 5;
+		FText ElapsedTimeMessage = FText::Format(LOCTEXT("ElapsedTime", "Checking {0} with {1} took {2} ms."), 
+			FText::FromString(Asset->GetPathName()),
+			FText::FromString(GetClass()->GetPathName()),
+			FText::AsNumber(ElapsedTimeMS, &TimeFormat));
+		UE_LOG(LogContentValidation, VeryVerbose, TEXT("%s"), *ElapsedTimeMessage.ToString());
+	}
 	Result = CombineDataValidationResults(Result, ExtractValidationState(Context)); // Extract messages and validation state from members (primarily for use by BP interface functions)
 
 	return Result;
@@ -95,11 +113,6 @@ void UEditorValidatorBase::AssetFails(UObject* InAsset, const FText& InMessage)
 
 	FText FailureMessage = FText::Format(LOCTEXT("AssetCheck_Message_Error", "{CustomMessage}. ({ValidatorName})"), Arguments);
 
-	if(LogContentValidation.GetVerbosity() >= ELogVerbosity::Verbose)
-	{
-		LogElapsedTime(Arguments);
-
-	}
 
 	AllErrors.Add(FailureMessage);
 	ValidationResult = EDataValidationResult::Invalid;
@@ -115,18 +128,6 @@ void UEditorValidatorBase::AssetWarning(UObject* InAsset, const FText& InMessage
 	AllWarnings.Add(WarningMessage);
 }
 
-void UEditorValidatorBase::LogElapsedTime(FFormatNamedArguments &Arguments)
-{
-	FDateTime CurrentTime = FDateTime::Now();
-	FTimespan ElapsedTimeSpan = (CurrentTime - ValidationTime);
-	float ElapsedTimeMS = ElapsedTimeSpan.GetTotalMilliseconds();
-	FNumberFormattingOptions TimeFormat;
-	TimeFormat.MinimumFractionalDigits = 5;
-	Arguments.Add(TEXT("ElapsedTime"), FText::AsNumber(ElapsedTimeMS, &TimeFormat));
-	FText ElapsedTimeMessage = FText::Format(LOCTEXT("ElapsedTime", "Checking {AssetName} with {ValidatorName} took {ElapsedTime} ms."), Arguments);
-	UE_LOG(LogContentValidation, Verbose, TEXT("%s"), *ElapsedTimeMessage.ToString());
-}
-
 void UEditorValidatorBase::AssetPasses(UObject* InAsset)
 {
 	if (LogContentValidation.GetVerbosity() >= ELogVerbosity::Verbose)
@@ -137,8 +138,6 @@ void UEditorValidatorBase::AssetPasses(UObject* InAsset)
 			Arguments.Add(TEXT("AssetName"), FText::FromName(InAsset->GetFName()));
 		}
 		Arguments.Add(TEXT("ValidatorName"), FText::FromString(GetClass()->GetPathName()));
-
-		LogElapsedTime(Arguments);
 	}
 
 	ensureMsgf(ValidationResult != EDataValidationResult::Invalid, TEXT("%s: AssetPasses called after errors were reported"), *GetClass()->GetPathName());
@@ -177,7 +176,6 @@ void UEditorValidatorBase::AddLegacyValidationErrors(TArray<FText> InErrors)
 void UEditorValidatorBase::ResetValidationState()
 {
 	ValidationResult = EDataValidationResult::NotValidated;
-	ValidationTime = FDateTime::Now();
 	AllMessages.Empty();
 	AllWarnings.Empty();
 	AllErrors.Empty();
diff --git a/Engine/Plugins/Editor/DataValidation/Source/DataValidation/Public/EditorValidatorBase.h b/Engine/Plugins/Editor/DataValidation/Source/DataValidation/Public/EditorValidatorBase.h
--- a/Engine/Plugins/Editor/DataValidation/Source/DataValidation/Public/EditorValidatorBase.h
+++ b/Engine/Plugins/Editor/DataValidation/Source/DataValidation/Public/EditorValidatorBase.h
@@ -132,9 +132,6 @@ public:
 		return EDataValidationResult::NotValidated;
 	}
 
-protected:
-	void LogElapsedTime(FFormatNamedArguments &Arguments);
-
 protected:
 	UPROPERTY(EditAnywhere, Category = "Asset Validation", meta = (BlueprintProtected = "true"))
 	bool bIsEnabled;
@@ -149,7 +146,6 @@ private:
 	TArray<FText> AllErrors;
 	TArray<TSharedRef<FTokenizedMessage>> AllMessages;
 
-	FDateTime ValidationTime;
 	
 	UPROPERTY()
 	UObject* CurrentObjectBeingValidated;
