#UE Fix an error where objects that are in the process of async loading are not being validated because they are not done by the time it is their turn

diff --git a/Engine/Plugins/Editor/DataValidation/Source/DataValidation/Private/EditorValidatorSubsystem.cpp b/Engine/Plugins/Editor/DataValidation/Source/DataValidation/Private/EditorValidatorSubsystem.cpp
--- a/Engine/Plugins/Editor/DataValidation/Source/DataValidation/Private/EditorValidatorSubsystem.cpp
+++ b/Engine/Plugins/Editor/DataValidation/Source/DataValidation/Private/EditorValidatorSubsystem.cpp
@@ -424,6 +424,9 @@ EDataValidationResult UEditorValidatorSubsystem::ValidateAssetsInternal(
 			++OutResults.NumSkipped;
 			continue;
 		}
+
+		// Flush here in case loading something in this loop caused something later in the loop to start async loading, but not finish before we get to it
+		FlushAsyncLoading();
 		
 		DataValidationLog.Info()
 			->AddToken(FAssetDataToken::Create(Data))
