Fixed validation of external actors, and exclude temp assets/actors from validation

#rnx

diff --git a/Engine/Plugins/Editor/DataValidation/Source/DataValidation/Private/EditorValidatorSubsystem.cpp b/Engine/Plugins/Editor/DataValidation/Source/DataValidation/Private/EditorValidatorSubsystem.cpp
--- a/Engine/Plugins/Editor/DataValidation/Source/DataValidation/Private/EditorValidatorSubsystem.cpp
+++ b/Engine/Plugins/Editor/DataValidation/Source/DataValidation/Private/EditorValidatorSubsystem.cpp
@@ -118,7 +118,15 @@ bool UEditorValidatorSubsystem::ShouldValidateAsset(
 	const FValidateAssetsSettings& 	Settings,
 	FDataValidationContext& 		InContext) const
 {
-	if (IsPathExcludedFromValidation(WriteToString<FName::StringBufferSize>(Asset.PackageName).ToView()))
+	FNameBuilder AssetPackageNameBuilder(Asset.PackageName);
+	FStringView AssetPackageNameView = AssetPackageNameBuilder.ToView();
+
+	if (FPackageName::IsTempPackage(AssetPackageNameView))
+	{
+		return false;
+	}
+
+	if (IsPathExcludedFromValidation(AssetPackageNameView))
 	{
 		return false;
 	}
@@ -365,7 +373,10 @@ EDataValidationResult UEditorValidatorSubsystem::ValidateAssetsInternal(
 		{
 			FSoftObjectPath Path = It->ToSoftObjectPath();
 			FAssetData OuterAsset = AssetRegistry.GetAssetByObjectPath(Path.GetWithoutSubPath(), true);
-			AssetsToExternalObjects.FindOrAdd(OuterAsset).Add(*It);
+			if (OuterAsset.IsValid())
+			{
+				AssetsToExternalObjects.FindOrAdd(OuterAsset).Add(*It);
+			}
 			It.RemoveCurrent();
 		}
 	}
@@ -433,7 +444,13 @@ EDataValidationResult UEditorValidatorSubsystem::ValidateAssetsInternal(
 		UObject* LoadedAsset = Data.FastGetAsset(false);
 		const bool bAlreadyLoaded = LoadedAsset != nullptr;
 
-		FDataValidationContext ValidationContext(bAlreadyLoaded, InSettings.ValidationUsecase, AssetsToExternalObjects.FindRef(Data));
+		TConstArrayView<FAssetData> ValidationExternalObjects;
+		if (const TArray<FAssetData>* ValidationExternalObjectsPtr = AssetsToExternalObjects.Find(Data))
+		{
+			ValidationExternalObjects = *ValidationExternalObjectsPtr;
+		}
+
+		FDataValidationContext ValidationContext(bAlreadyLoaded, InSettings.ValidationUsecase, ValidationExternalObjects);
 		EDataValidationResult AssetResult = EDataValidationResult::NotValidated;
 		if (!LoadedAsset)
 		{
diff --git a/Engine/Source/Runtime/CoreUObject/Public/Misc/DataValidation.h b/Engine/Source/Runtime/CoreUObject/Public/Misc/DataValidation.h
--- a/Engine/Source/Runtime/CoreUObject/Public/Misc/DataValidation.h
+++ b/Engine/Source/Runtime/CoreUObject/Public/Misc/DataValidation.h
@@ -66,7 +66,8 @@ public:
 		bool InWasAssetLoadedForValidation,
 		EDataValidationUsecase InUsecase,
 		TConstArrayView<FAssetData> InAssociatedObjects)
-		: ValidationUsecase(InUsecase)
+		: AssociatedExternalObjects(InAssociatedObjects)
+		, ValidationUsecase(InUsecase)
 		, bWasAssetLoadedForValidation(InWasAssetLoadedForValidation)
 	{}
 
