Skip UDataValidationChangelist::IsDataValid on build machines

diff --git a/Engine/Plugins/Editor/DataValidation/Source/DataValidation/Private/DataValidationChangelist.cpp b/Engine/Plugins/Editor/DataValidation/Source/DataValidation/Private/DataValidationChangelist.cpp
--- a/Engine/Plugins/Editor/DataValidation/Source/DataValidation/Private/DataValidationChangelist.cpp
+++ b/Engine/Plugins/Editor/DataValidation/Source/DataValidation/Private/DataValidationChangelist.cpp
@@ -75,6 +75,13 @@ FString UDataValidationChangelist::GetPrettyPackageName(const FName& InPackageNa
 
 EDataValidationResult UDataValidationChangelist::IsDataValid(FDataValidationContext& Context) const
 {
+	// Temporary: do not validate changelists objects on build machines
+	// In future we should differentiate between pending and submitted changelists to change the behavior of this function
+	if (GIsBuildMachine)
+	{
+		return EDataValidationResult::Valid;
+	}
+
 	ISourceControlProvider& SourceControlProvider = ISourceControlModule::Get().GetProvider();
 
 	// Gather dependencies of every file in the changelist
