Logging

diff --git a/Engine/Plugins/Editor/DataValidation/Source/DataValidation/Private/EditorValidatorSubsystem.cpp b/Engine/Plugins/Editor/DataValidation/Source/DataValidation/Private/EditorValidatorSubsystem.cpp
--- a/Engine/Plugins/Editor/DataValidation/Source/DataValidation/Private/EditorValidatorSubsystem.cpp
+++ b/Engine/Plugins/Editor/DataValidation/Source/DataValidation/Private/EditorValidatorSubsystem.cpp
@@ -105,9 +105,9 @@ void UEditorValidatorSubsystem::Initialize(FSubsystemCollectionBase& Collection)
 					UEditorValidatorBase* Validator = NewObject<UEditorValidatorBase>(GetTransientPackage(), ValidatorClass);
 					AddValidator(Validator);
 				}
+				}
 			}
 		}
-	}
 
 	// Register to SCC pre-submit callback
 	ISourceControlModule::Get().RegisterPreSubmitDataValidation(FSourceControlPreSubmitDataValidationDelegate::CreateUObject(this, &UEditorValidatorSubsystem::ValidateChangelistPreSubmit));
@@ -147,21 +147,21 @@ void UEditorValidatorSubsystem::RegisterBlueprintValidators()
 		for (FAssetData& BPAssetData : AllBPsAssetData)
 		{
 			UClass* ParentClass = nullptr;
-			FString ParentClassName;
+	FString ParentClassName;
 			if (!BPAssetData.GetTagValue(FBlueprintTags::NativeParentClassPath, ParentClassName))
-			{
+	{
 				BPAssetData.GetTagValue(FBlueprintTags::ParentClassPath, ParentClassName);
-			}
+	}
 
 			if (!ParentClassName.IsEmpty())
-			{
-				UObject* Outer = nullptr;
-				ResolveName(Outer, ParentClassName, false, false);
-				ParentClass = FindObject<UClass>(Outer, *ParentClassName);
-				if (!ParentClass || !ParentClass->IsChildOf(UEditorValidatorBase::StaticClass()))
-				{
+	{
+	UObject* Outer = nullptr;
+	ResolveName(Outer, ParentClassName, false, false);
+	ParentClass = FindObject<UClass>(Outer, *ParentClassName);
+	if (!ParentClass || !ParentClass->IsChildOf(UEditorValidatorBase::StaticClass()))
+	{
 					continue;
-				}
+	}
 			}
 
 			AddValidator(BPAssetData);
@@ -351,7 +351,7 @@ EDataValidationResult UEditorValidatorSubsystem::ValidateAssetsInternal(
 	SlowTask.MakeDialog();
 	
 	LoadValidators();
-	UE_LOG(LogContentValidation, Log, TEXT("Starting to validate %d assets"), AssetDataList.Num());
+	UE_LOG(LogContentValidation, Display, TEXT("Starting to validate %d assets"), AssetDataList.Num());
 	UE_LOG(LogContentValidation, Log, TEXT("Registered validators:"));
 	for (const TPair<FTopLevelAssetPath, TObjectPtr<UEditorValidatorBase>>& ValidatorPair : Validators)
 	{
@@ -439,7 +439,7 @@ EDataValidationResult UEditorValidatorSubsystem::ValidateAssetsInternal(
 		DataValidationLog.Info()
 			->AddToken(FAssetDataToken::Create(Data))
 			->AddToken(FTextToken::Create(LOCTEXT("Data.ValidatingAsset", "Validating asset")));
-		UE_LOG(LogContentValidation, Log, TEXT("Validating asset %s"), *Data.ToSoftObjectPath().ToString());
+		UE_LOG(LogContentValidation, Display, TEXT("Validating asset %s"), *Data.ToSoftObjectPath().ToString());
 		
 		UObject* LoadedAsset = Data.FastGetAsset(false);
 		const bool bAlreadyLoaded = LoadedAsset != nullptr;
