Fix extraction of warnings and errors from validation contexts

diff --git a/Engine/Plugins/Editor/DataValidation/Source/DataValidation/Private/EditorValidatorSubsystem.cpp b/Engine/Plugins/Editor/DataValidation/Source/DataValidation/Private/EditorValidatorSubsystem.cpp
--- a/Engine/Plugins/Editor/DataValidation/Source/DataValidation/Private/EditorValidatorSubsystem.cpp
+++ b/Engine/Plugins/Editor/DataValidation/Source/DataValidation/Private/EditorValidatorSubsystem.cpp
@@ -220,7 +220,7 @@ EDataValidationResult UEditorValidatorSubsystem::IsObjectValid(
 {
 	FDataValidationContext Context(false, InValidationUsecase, {}); // No associated objects in this context
 	EDataValidationResult Result = IsObjectValidWithContext(InObject, Context);
-	Context.SplitIssues(ValidationErrors, ValidationWarnings);
+	Context.SplitIssues(ValidationWarnings, ValidationErrors);
 	return Result;
 }
 
@@ -237,7 +237,7 @@ EDataValidationResult UEditorValidatorSubsystem::IsAssetValid(
 		{
 			FDataValidationContext Context(false, InValidationUsecase, {}); // No associated objects in this context
 			EDataValidationResult Result = ValidateObjectInternal(AssetData, Obj, Context);
-			Context.SplitIssues(ValidationErrors, ValidationWarnings);
+			Context.SplitIssues(ValidationWarnings, ValidationErrors);
 			return Result;
 		}
 		return EDataValidationResult::NotValidated;
@@ -579,7 +579,7 @@ EDataValidationResult UEditorValidatorSubsystem::ValidateAssetsInternal(
 			Details.PackageName = Data.PackageName;
 			Details.AssetName = Data.AssetName;
 			Details.Result = AssetResult;
-			ValidationContext.SplitIssues(Details.ValidationErrors, Details.ValidationWarnings);
+			ValidationContext.SplitIssues(Details.ValidationWarnings, Details.ValidationErrors);
 		}
 		
 		DataValidationLog.Flush();
