Allow "Validate" to work on newly added actors

The query for the asset data would fail in this case as there was no file on disk, so in this case we continue to use the asset data generated from the in-memory actor.

#rnx

diff --git a/Engine/Plugins/Editor/DataValidation/Source/DataValidation/Private/DataValidationModule.cpp b/Engine/Plugins/Editor/DataValidation/Source/DataValidation/Private/DataValidationModule.cpp
--- a/Engine/Plugins/Editor/DataValidation/Source/DataValidation/Private/DataValidationModule.cpp
+++ b/Engine/Plugins/Editor/DataValidation/Source/DataValidation/Private/DataValidationModule.cpp
@@ -117,7 +117,11 @@ void FDataValidationModule::RegisterMenus()
 				IAssetRegistry& AssetRegistry = AssetRegistryModule.Get();
 				for (FAssetData& AssetData : SelectedActorAssets)
 				{
-					AssetData = AssetRegistry.GetAssetByObjectPath(AssetData.GetSoftObjectPath(), true);
+					FAssetData DiskAssetData = AssetRegistry.GetAssetByObjectPath(AssetData.GetSoftObjectPath(), true);
+					if (DiskAssetData.IsValid())
+					{
+						AssetData = MoveTemp(DiskAssetData);
+					}
 				}
 				ValidateAssets(SelectedActorAssets, false, EDataValidationUsecase::Manual);
 			}
