Make AssetFails, AssetWarning and AssetPasses functions use const object parameters

diff --git a/Engine/Plugins/Editor/DataValidation/Source/DataValidation/Private/EditorValidatorBase.cpp b/Engine/Plugins/Editor/DataValidation/Source/DataValidation/Private/EditorValidatorBase.cpp
--- a/Engine/Plugins/Editor/DataValidation/Source/DataValidation/Private/EditorValidatorBase.cpp
+++ b/Engine/Plugins/Editor/DataValidation/Source/DataValidation/Private/EditorValidatorBase.cpp
@@ -96,7 +96,7 @@ EDataValidationResult UEditorValidatorBase::K2_ValidateLoadedAsset_Implementatio
 	return EDataValidationResult::NotValidated;
 }
 
-void UEditorValidatorBase::AssetFails(UObject* InAsset, const FText& InMessage, TArray<FText>& InOutErrors)
+void UEditorValidatorBase::AssetFails(const UObject* InAsset, const FText& InMessage, TArray<FText>& InOutErrors)
 {
 	AssetFails(InAsset, InMessage);
 	if (&InOutErrors != &AllErrors)
@@ -105,7 +105,7 @@ void UEditorValidatorBase::AssetFails(UObject* InAsset, const FText& InMessage,
 	}
 }
 
-void UEditorValidatorBase::AssetFails(UObject* InAsset, const FText& InMessage)
+void UEditorValidatorBase::AssetFails(const UObject* InAsset, const FText& InMessage)
 {
 	FFormatNamedArguments Arguments;
 	Arguments.Add(TEXT("CustomMessage"), InMessage);
@@ -118,7 +118,7 @@ void UEditorValidatorBase::AssetFails(UObject* InAsset, const FText& InMessage)
 	ValidationResult = EDataValidationResult::Invalid;
 }
 
-void UEditorValidatorBase::AssetWarning(UObject* InAsset, const FText& InMessage)
+void UEditorValidatorBase::AssetWarning(const UObject* InAsset, const FText& InMessage)
 {
 	FFormatNamedArguments Arguments;
 	Arguments.Add(TEXT("CustomMessage"), InMessage);
@@ -128,7 +128,7 @@ void UEditorValidatorBase::AssetWarning(UObject* InAsset, const FText& InMessage
 	AllWarnings.Add(WarningMessage);
 }
 
-void UEditorValidatorBase::AssetPasses(UObject* InAsset)
+void UEditorValidatorBase::AssetPasses(const UObject* InAsset)
 {
 	if (LogContentValidation.GetVerbosity() >= ELogVerbosity::Verbose)
 	{
diff --git a/Engine/Plugins/Editor/DataValidation/Source/DataValidation/Public/EditorValidatorBase.h b/Engine/Plugins/Editor/DataValidation/Source/DataValidation/Public/EditorValidatorBase.h
--- a/Engine/Plugins/Editor/DataValidation/Source/DataValidation/Public/EditorValidatorBase.h
+++ b/Engine/Plugins/Editor/DataValidation/Source/DataValidation/Public/EditorValidatorBase.h
@@ -73,18 +73,18 @@ public:
 	
 	/** Marks the validation as failed and adds an error message. */
 	UFUNCTION(BlueprintCallable, Category = "Asset Validation")
-	void AssetFails(UObject* InAsset, const FText& InMessage);
+	void AssetFails(const UObject* InAsset, const FText& InMessage);
 
 	UE_DEPRECATED("5.4", "This function has been replaced by the version AssetFails(UObject*, const FText&)")
-	void AssetFails(UObject* InAsset, const FText& InMessage, TArray<FText>& InOutErrors);
+	void AssetFails(const UObject* InAsset, const FText& InMessage, TArray<FText>& InOutErrors);
 
 	/** Marks the validation as successful. Failure to call this will report the validator as not having checked the asset. */
 	UFUNCTION(BlueprintCallable, Category = "Asset Validation")
-	void AssetPasses(UObject* InAsset);
+	void AssetPasses(const UObject* InAsset);
 
 	/** Adds a message to this validation but doesn't mark it as failed. */
 	UFUNCTION(BlueprintCallable, Category = "Asset Validation")
-	void AssetWarning(UObject* InAsset, const FText& InMessage);
+	void AssetWarning(const UObject* InAsset, const FText& InMessage);
 	
 	/** 
 	 * Add a tokenized message to the validation results. If the severity is error, marks the validation as failed.
