Fix Validation failed check when validating a Level Instance while it is in Edit

Avoid using loaded world in this case as it will not be a main container and instead create a container instance

#rnx

diff --git a/Engine/Plugins/Editor/DataValidation/Source/DataValidation/Private/WorldPartitionChangelistValidator.cpp b/Engine/Plugins/Editor/DataValidation/Source/DataValidation/Private/WorldPartitionChangelistValidator.cpp
--- a/Engine/Plugins/Editor/DataValidation/Source/DataValidation/Private/WorldPartitionChangelistValidator.cpp
+++ b/Engine/Plugins/Editor/DataValidation/Source/DataValidation/Private/WorldPartitionChangelistValidator.cpp
@@ -143,9 +143,11 @@ void UWorldPartitionChangelistValidator::ValidateActorsAndDataLayersFromChangeLi
 		}
 
 		UActorDescContainerInstance* ContainerInstance = nullptr;
-		if (InWorld != nullptr)
+		if (InWorld != nullptr && InWorld->GetWorldPartition()->IsMainWorldPartition())
 		{
-			// World is Loaded reuse the ActorDescContainer of the Content Bundle
+			// World is Loaded reuse the ActorDescContainer (Can be Main world package or one of the Content Bundle packages associated with main world)
+			// Only reuse Main world partiton worlds because streaming generation expects the top level container to always be a Main container.
+			// It is possible that InWorld is a LevelInstance edit world in which case we don't want to use its loaded ContainerInstance and instead create a new one
 			ContainerInstance = InWorld->GetWorldPartition()->FindContainer(InContainerPackageName);
 		}
 
