Fix double semicolon.

diff --git a/Engine/Plugins/Editor/DataValidation/Source/DataValidation/Private/WorldPartitionChangelistValidator.cpp b/Engine/Plugins/Editor/DataValidation/Source/DataValidation/Private/WorldPartitionChangelistValidator.cpp
--- a/Engine/Plugins/Editor/DataValidation/Source/DataValidation/Private/WorldPartitionChangelistValidator.cpp
+++ b/Engine/Plugins/Editor/DataValidation/Source/DataValidation/Private/WorldPartitionChangelistValidator.cpp
@@ -79,7 +79,7 @@ void UWorldPartitionChangelistValidator::ValidateActorsAndDataLayersFromChangeLi
 			if (ActorFiles)	// A null Files indicates a World not using World Partition and OFPA 
 			{
 				ActorFiles->Add(AssetData);
-				return FWorldPartitionActorDescUtils::GetActorNativeClassFromAssetData(AssetData);;
+				return FWorldPartitionActorDescUtils::GetActorNativeClassFromAssetData(AssetData);
 			}
 		}
 
