Remove stale bValidateAssetsWhileSavingForCook flag

This flag was never used internally and no longer functioned as expected

#rnx

diff --git a/Engine/Plugins/Editor/DataValidation/Source/DataValidation/Private/EditorValidatorSubsystem.cpp b/Engine/Plugins/Editor/DataValidation/Source/DataValidation/Private/EditorValidatorSubsystem.cpp
--- a/Engine/Plugins/Editor/DataValidation/Source/DataValidation/Private/EditorValidatorSubsystem.cpp
+++ b/Engine/Plugins/Editor/DataValidation/Source/DataValidation/Private/EditorValidatorSubsystem.cpp
@@ -626,7 +626,7 @@ void UEditorValidatorSubsystem::ValidateOnSave(TArray<FAssetData> AssetDataList,
 		return;
 	}
 
-	if ((!bValidateAssetsWhileSavingForCook && bProceduralSave))
+	if (bProceduralSave)
 	{
 		return;
 	}
@@ -668,7 +668,7 @@ void UEditorValidatorSubsystem::ValidateSavedPackage(FName PackageName, bool bPr
 	}
 
 	// For performance reasons, don't validate when making a procedural save by default. Assumption is we validated when saving previously. 
-	if ((!bValidateAssetsWhileSavingForCook && bProceduralSave))
+	if (bProceduralSave)
 	{
 		return;
 	}
diff --git a/Engine/Plugins/Editor/DataValidation/Source/DataValidation/Public/EditorValidatorSubsystem.h b/Engine/Plugins/Editor/DataValidation/Source/DataValidation/Public/EditorValidatorSubsystem.h
--- a/Engine/Plugins/Editor/DataValidation/Source/DataValidation/Public/EditorValidatorSubsystem.h
+++ b/Engine/Plugins/Editor/DataValidation/Source/DataValidation/Public/EditorValidatorSubsystem.h
@@ -380,10 +380,6 @@ protected:
 	UPROPERTY(Transient)
 	TMap<FTopLevelAssetPath, TObjectPtr<UEditorValidatorBase>> Validators;
 
-	/** Specifies whether or not to validate assets on save when saving for a cook */
-	UPROPERTY(config)
-	bool bValidateAssetsWhileSavingForCook;
-
 	/** Specifies whether or not to allow Blueprint validators */
 	UPROPERTY(config)
 	bool bAllowBlueprintValidators;
