Pass back external objects in FValidateAssetsDetails

diff --git a/Engine/Plugins/Editor/DataValidation/Source/DataValidation/Private/EditorValidatorSubsystem.cpp b/Engine/Plugins/Editor/DataValidation/Source/DataValidation/Private/EditorValidatorSubsystem.cpp
--- a/Engine/Plugins/Editor/DataValidation/Source/DataValidation/Private/EditorValidatorSubsystem.cpp
+++ b/Engine/Plugins/Editor/DataValidation/Source/DataValidation/Private/EditorValidatorSubsystem.cpp
@@ -580,6 +580,14 @@ EDataValidationResult UEditorValidatorSubsystem::ValidateAssetsInternal(
 			Details.AssetName = Data.AssetName;
 			Details.Result = AssetResult;
 			ValidationContext.SplitIssues(Details.ValidationWarnings, Details.ValidationErrors);
+
+			Details.ExternalObjects.Reserve(ValidationExternalObjects.Num());
+			for (const FAssetData& ExtData : ValidationExternalObjects)
+			{
+				FValidateAssetsExternalObject& ExtDetails = Details.ExternalObjects.Emplace_GetRef();
+				ExtDetails.PackageName = ExtData.PackageName;
+				ExtDetails.AssetName = ExtData.AssetName;
+			}
 		}
 		
 		DataValidationLog.Flush();
diff --git a/Engine/Plugins/Editor/DataValidation/Source/DataValidation/Public/EditorValidatorSubsystem.h b/Engine/Plugins/Editor/DataValidation/Source/DataValidation/Public/EditorValidatorSubsystem.h
--- a/Engine/Plugins/Editor/DataValidation/Source/DataValidation/Public/EditorValidatorSubsystem.h
+++ b/Engine/Plugins/Editor/DataValidation/Source/DataValidation/Public/EditorValidatorSubsystem.h
@@ -27,6 +27,15 @@ namespace EMessageSeverity
 
 DECLARE_LOG_CATEGORY_EXTERN(LogContentValidation, Log, All);
 
+struct DATAVALIDATION_API FValidateAssetsExternalObject
+{
+	/** Package Name */
+	FName PackageName;
+
+	/** Asset Name */
+	FName AssetName;
+};
+
 USTRUCT(BlueprintType)
 struct DATAVALIDATION_API FValidateAssetsDetails
 {
@@ -49,6 +58,9 @@ struct DATAVALIDATION_API FValidateAssetsDetails
 	
 	/** Rich validation messages including tokens that can generate images or hyperlinks */
 	TArray<TSharedRef<FTokenizedMessage>> ValidationMessages;
+
+	/** List of external objects for this asset*/
+	TArray<FValidateAssetsExternalObject> ExternalObjects;
 };
 
 USTRUCT(BlueprintType)
