Capture UE_LOG messages associated with asset validation during cook

#rnx

diff --git a/Engine/Plugins/Editor/DataValidation/Source/DataValidation/Private/DataValidationModule.cpp b/Engine/Plugins/Editor/DataValidation/Source/DataValidation/Private/DataValidationModule.cpp
--- a/Engine/Plugins/Editor/DataValidation/Source/DataValidation/Private/DataValidationModule.cpp
+++ b/Engine/Plugins/Editor/DataValidation/Source/DataValidation/Private/DataValidationModule.cpp
@@ -400,7 +400,32 @@ EDataValidationResult FDataValidationModule::OnValidateSourcePackageDuringCook(U
 						->AddToken(FAssetDataToken::Create(AssetData))
 						->AddToken(FTextToken::Create(LOCTEXT("Data.ValidatingAsset", "Validating asset")));
 
-					const EDataValidationResult ValidationResult = EditorValidationSubsystem->IsObjectValidWithContext(Asset, ValidationContext);
+					UE::DataValidation::FScopedLogMessageGatherer LogGatherer;
+					EDataValidationResult ValidationResult = EditorValidationSubsystem->IsObjectValidWithContext(Asset, ValidationContext);
+
+					TArray<FString> LogWarnings;
+					TArray<FString> LogErrors;
+					LogGatherer.Stop(LogWarnings, LogErrors);
+
+					if (LogWarnings.Num() > 0)
+					{
+						TStringBuilder<2048> Buffer;
+						Buffer.Join(LogWarnings, LINE_TERMINATOR);
+						ValidationContext.AddMessage(EMessageSeverity::Error)
+							->AddToken(FAssetDataToken::Create(AssetData))
+							->AddText(LOCTEXT("DataValidation.DuringValidationWarnings", "Warnings logged while validating asset {0}"), FText::FromStringView(Buffer.ToView()));
+						ValidationResult = EDataValidationResult::Invalid;
+					}
+					if (LogErrors.Num() > 0)
+					{
+						TStringBuilder<2048> Buffer;
+						Buffer.Join(LogErrors, LINE_TERMINATOR);
+						ValidationContext.AddMessage(EMessageSeverity::Error)
+							->AddToken(FAssetDataToken::Create(AssetData))
+							->AddText(LOCTEXT("DataValidation.DuringValidationErrors", "Errors logged while validating asset {0}"), FText::FromStringView(Buffer.ToView()));
+						ValidationResult = EDataValidationResult::Invalid;
+					}
+
 					FinalValidationResult = CombineDataValidationResults(FinalValidationResult, ValidationResult);
 
 					UE::DataValidation::AddAssetValidationMessages(AssetData, DataValidationLog, ValidationContext);
