Editor | Add functionality for automatic fixes in tokenized messages

This allows you to attach an automatic fix to any error or warning that may arise from eg. validation.

diff --git a/Engine/Plugins/Editor/DataValidation/Source/DataValidation/Private/DataValidationFixers.cpp b/Engine/Plugins/Editor/DataValidation/Source/DataValidation/Private/DataValidationFixers.cpp
new file mode 100644
--- /dev/null
+++ b/Engine/Plugins/Editor/DataValidation/Source/DataValidation/Private/DataValidationFixers.cpp
@@ -0,0 +1,219 @@
+﻿// Copyright Epic Games, Inc. All Rights Reserved.
+
+#include "DataValidationFixers.h"
+
+#include "Editor.h"
+#include "EditorValidatorSubsystem.h"
+#include "FileHelpers.h"
+#include "Algo/AnyOf.h"
+#include "UObject/Package.h"
+
+#define LOCTEXT_NAMESPACE "DataValidationFixers"
+
+namespace UE::DataValidation
+{
+
+EFixApplicability FSingleUseFixer::GetApplicability(int32 FixIndex) const
+{
+	if (UsedFixes.Contains(FixIndex))
+	{
+		return EFixApplicability::Applied;
+	}
+	else
+	{
+		return Inner->GetApplicability(FixIndex);
+	}
+}
+
+FFixResult FSingleUseFixer::ApplyFix(int32 FixIndex)
+{
+	FFixResult Result = Inner->ApplyFix(FixIndex);
+	UsedFixes.Add(FixIndex);
+	return Result;
+}
+
+TSharedRef<FSingleUseFixer> FSingleUseFixer::Create(TSharedRef<IFixer> Inner)
+{
+	TSharedRef<FSingleUseFixer> Fixer = MakeShared<FSingleUseFixer>();
+	Fixer->Inner = MoveTemp(Inner);
+	return Fixer;
+}
+
+EFixApplicability FObjectSetDependentFixer::GetApplicability(int32 FixIndex) const
+{
+	if (Algo::AnyOf(Dependencies, [](const TWeakObjectPtr<>& Ptr) { return Ptr.IsStale(); }))
+	{
+		return EFixApplicability::DidNotApply;
+	}
+	else
+	{
+		return Inner->GetApplicability(FixIndex);
+	}
+}
+
+FFixResult FObjectSetDependentFixer::ApplyFix(int32 FixIndex)
+{
+	return Inner->ApplyFix(FixIndex);
+}
+
+TSharedRef<FObjectSetDependentFixer> FObjectSetDependentFixer::Create(TSharedRef<IFixer> Inner, TArray<TWeakObjectPtr<>> Dependencies)
+{
+	TSharedRef<FObjectSetDependentFixer> Fixer = MakeShared<FObjectSetDependentFixer>();
+	Fixer->Inner = MoveTemp(Inner);
+	Fixer->Dependencies = MoveTemp(Dependencies);
+	return Fixer;
+}
+
+EFixApplicability FAutoSavingFixer::GetApplicability(int32 FixIndex) const
+{
+	return Inner->GetApplicability(FixIndex);
+}
+
+template <typename F>
+auto CollectDirtyPackagesDuring(TSet<UPackage*>& OutPackagesMarkedDirty, F&& Lambda) -> decltype(Lambda())
+{
+	FDelegateHandle PackageModificationListenerHandle = UPackage::PackageMarkedDirtyEvent.AddLambda(
+		// Lambda capture safety: This lambda is removed from the delegate before PackagesMarkedDirty goes out of scope.
+		[&OutPackagesMarkedDirty](UPackage* Package, bool bIsDirty)
+		{
+			if (bIsDirty)
+			{
+				OutPackagesMarkedDirty.Add(Package);
+			}
+		}
+	);
+
+	auto Result = Lambda();
+
+	UPackage::PackageMarkedDirtyEvent.Remove(PackageModificationListenerHandle);
+
+	return Result;
+}
+
+FFixResult FAutoSavingFixer::ApplyFix(int32 FixIndex)
+{
+	TSet<UPackage*> PackagesMarkedDirty;
+	FFixResult FixResult =
+		CollectDirtyPackagesDuring(PackagesMarkedDirty, [this, FixIndex] { return Inner->ApplyFix(FixIndex); });
+
+	if (FixResult.bIsSuccess && !PackagesMarkedDirty.IsEmpty())
+	{
+		TArray<UPackage*> PackagesToSave = PackagesMarkedDirty.Array();
+		FEditorFileUtils::FPromptForCheckoutAndSaveParams Params;
+		Params.Title = LOCTEXT("SaveFixedAssets", "Save assets after applying fix");
+		Params.Message = LOCTEXT(
+			"SaveFixedAssetsDescription", "Applying the fix modified the following assets, which now need to be saved."
+		);
+		// NOTE: Return code ignored because ultimately the automatic save is just a nice thing we do for the user.
+		// It's not mandatory for the assets to be saved.
+		[[maybe_unused]] FEditorFileUtils::EPromptReturnCode PromptReturnCode =
+			FEditorFileUtils::PromptForCheckoutAndSave(PackagesToSave, Params);
+	}
+
+	return FixResult;
+}
+
+TSharedRef<FAutoSavingFixer> FAutoSavingFixer::Create(TSharedRef<IFixer> Inner)
+{
+	TSharedRef<FAutoSavingFixer> Fixer = MakeShared<FAutoSavingFixer>();
+	Fixer->Inner = MoveTemp(Inner);
+	return Fixer;
+}
+
+EFixApplicability FValidatingFixer::GetApplicability(int32 FixIndex) const
+{
+	return Inner->GetApplicability(FixIndex);
+}
+
+FFixResult FValidatingFixer::ApplyFix(int32 FixIndex)
+{
+	TSet<UPackage*> PackagesMarkedDirty;
+	FFixResult FixResult =
+		CollectDirtyPackagesDuring(PackagesMarkedDirty, [this, FixIndex] { return Inner->ApplyFix(FixIndex); });
+
+	if (FixResult.bIsSuccess && !PackagesMarkedDirty.IsEmpty())
+	{
+		if (GEditor != nullptr)
+		{
+			UEditorValidatorSubsystem* ValidatorSubsystem = GEditor->GetEditorSubsystem<UEditorValidatorSubsystem>();
+			TArray<FAssetData> AssetsToValidate;
+			AssetsToValidate.Reserve(PackagesMarkedDirty.Num());
+			for (const UPackage* Package : PackagesMarkedDirty)
+			{
+				AssetsToValidate.Add(FAssetData(Package));
+			}
+			FValidateAssetsSettings ValidationSettings;
+			ValidationSettings.ValidationUsecase = EDataValidationUsecase::Save;
+			FValidateAssetsResults ValidationResults;
+			ValidatorSubsystem->ValidateAssetsWithSettings(AssetsToValidate, ValidationSettings, ValidationResults);
+		}
+	}
+
+	return FixResult;
+}
+
+TSharedRef<FValidatingFixer> FValidatingFixer::Create(TSharedRef<IFixer> Inner)
+{
+	TSharedRef<FValidatingFixer> Fixer = MakeShared<FValidatingFixer>();
+	Fixer->Inner = MoveTemp(Inner);
+	return Fixer;
+}
+
+FMutuallyExclusiveFixSet::FMutuallyExclusiveFixSet()
+{
+	SharedData = MakeShared<FSharedData>();
+}
+
+void FMutuallyExclusiveFixSet::Add(const FText& Label, TSharedRef<IFixer> Inner)
+{
+	TSharedRef<FFixer> Fixer = MakeShared<FFixer>();
+	Fixer->Inner = MoveTemp(Inner);
+	Fixer->SharedData = SharedData;
+	QueuedTokens.Add(FQueuedToken{.Label = Label, .Fixer = MoveTemp(Fixer)});
+}
+
+void FMutuallyExclusiveFixSet::Transform(const TFunctionRef<TSharedRef<IFixer>(TSharedRef<IFixer>)>& Callback)
+{
+	for (FQueuedToken& QueuedToken : QueuedTokens)
+	{
+		QueuedToken.Fixer->Inner = Callback(QueuedToken.Fixer->Inner.ToSharedRef());
+	}
+}
+
+void FMutuallyExclusiveFixSet::CreateTokens(const TFunctionRef<void(TSharedRef<FFixToken>)>& Callback) const
+{
+	for (int32 i = 0; i < QueuedTokens.Num(); ++i)
+	{
+		const FQueuedToken& QueuedToken = QueuedTokens[i];
+		FText Label = i == 0 ? LOCTEXT("FirstMutuallyExclusiveFix", "Fix: {0}") : LOCTEXT("NextMutuallyExclusiveFix", " or: {0}");
+		Callback(FFixToken::Create(FText::Format(Label, QueuedToken.Label), QueuedToken.Fixer, i));
+	}
+}
+
+EFixApplicability FMutuallyExclusiveFixSet::FFixer::GetApplicability(int32 FixIndex) const
+{
+	if (SharedData->AppliedFix == INDEX_NONE)
+	{
+		return Inner->GetApplicability(0);
+	}
+
+	if (SharedData->AppliedFix == FixIndex)
+	{
+		return EFixApplicability::Applied;
+	}
+	else
+	{
+		return EFixApplicability::DidNotApply;
+	}
+}
+
+FFixResult FMutuallyExclusiveFixSet::FFixer::ApplyFix(int32 FixIndex)
+{
+	FFixResult Result = Inner->ApplyFix(0);
+	SharedData->AppliedFix = FixIndex;
+	return Result;
+}
+
+}
+
+#undef LOCTEXT_NAMESPACE
diff --git a/Engine/Plugins/Editor/DataValidation/Source/DataValidation/Public/DataValidationFixers.h b/Engine/Plugins/Editor/DataValidation/Source/DataValidation/Public/DataValidationFixers.h
new file mode 100644
--- /dev/null
+++ b/Engine/Plugins/Editor/DataValidation/Source/DataValidation/Public/DataValidationFixers.h
@@ -0,0 +1,160 @@
+﻿// Copyright Epic Games, Inc. All Rights Reserved.
+
+#pragma once
+
+#include "Logging/TokenizedMessage.h"
+#include "Misc/DataValidation/Fixer.h"
+#include "Templates/SharedPointer.h"
+
+namespace UE::DataValidation
+{
+
+/** Functor which always returns EFixApplicability::CanBeApplied. */
+struct FFixAlwaysApplicable
+{
+	EFixApplicability operator()() const
+	{
+		return EFixApplicability::CanBeApplied;
+	}
+};
+
+/**
+ * IFixer created out of a pair of lambdas.
+ * This can be constructed conveniently using MakeFix.
+ */
+template <typename GetApplicabilityT, typename ApplyFixT>
+struct TLambdaFixer : IFixer
+{
+	TLambdaFixer(GetApplicabilityT&& GetApplicabilityImpl, ApplyFixT&& ApplyFixImpl)
+		: GetApplicabilityImpl(MoveTemp(GetApplicabilityImpl))
+		, ApplyFixImpl(MoveTemp(ApplyFixImpl))
+	{}
+
+	virtual EFixApplicability GetApplicability(int32 FixIndex) const override
+	{
+		return GetApplicabilityImpl();
+	}
+
+	virtual FFixResult ApplyFix(int32 FixIndex) override
+	{
+		return ApplyFixImpl();
+	}
+
+private:
+	GetApplicabilityT GetApplicabilityImpl;
+	ApplyFixT ApplyFixImpl;
+};
+
+/** Makes an IFixer from applicability and application lambdas. */
+template <typename GetApplicabilityT, typename ApplyFixT>
+TSharedRef<TLambdaFixer<GetApplicabilityT, ApplyFixT>> MakeFix(GetApplicabilityT&& GetApplicability, ApplyFixT&& ApplyFix)
+{
+	return MakeShared<TLambdaFixer<GetApplicabilityT, ApplyFixT>>(MoveTemp(GetApplicability), MoveTemp(ApplyFix));
+}
+
+/** Makes an always applicable IFixer from an application lambda. */
+template <typename ApplyFixT>
+TSharedRef<TLambdaFixer<FFixAlwaysApplicable, ApplyFixT>> MakeFix(ApplyFixT&& ApplyFix)
+{
+	return MakeShared<TLambdaFixer<FFixAlwaysApplicable, ApplyFixT>>(FFixAlwaysApplicable{}, MoveTemp(ApplyFix));
+}
+
+/**
+ * IFixer which wraps another fixer and makes it single use only.
+ * All fixers which are not idempotent will want to use this.
+ */
+struct DATAVALIDATION_API FSingleUseFixer : IFixer
+{
+	TSharedPtr<IFixer> Inner;
+	TSet<int32> UsedFixes;
+
+	virtual EFixApplicability GetApplicability(int32 FixIndex) const override;
+	virtual FFixResult ApplyFix(int32 FixIndex) override;
+
+	static TSharedRef<FSingleUseFixer> Create(TSharedRef<IFixer> Inner);
+};
+
+/** IFixer which wraps another fixer in a set of dependencies. */
+struct DATAVALIDATION_API FObjectSetDependentFixer : IFixer
+{
+	TSharedPtr<IFixer> Inner;
+	TArray<TWeakObjectPtr<>> Dependencies;
+
+	virtual EFixApplicability GetApplicability(int32 FixIndex) const override;
+	virtual FFixResult ApplyFix(int32 FixIndex) override;
+
+	static TSharedRef<FObjectSetDependentFixer> Create(TSharedRef<IFixer> Inner, TArray<TWeakObjectPtr<>> Dependencies);
+};
+
+/** IFixer which wraps another fixer and makes it auto-save the asset after the fix is applied. */
+struct DATAVALIDATION_API FAutoSavingFixer : IFixer
+{
+	TSharedPtr<IFixer> Inner;
+
+	virtual EFixApplicability GetApplicability(int32 FixIndex) const override;
+	virtual FFixResult ApplyFix(int32 FixIndex) override;
+
+	static TSharedRef<FAutoSavingFixer> Create(TSharedRef<IFixer> Inner);
+};
+
+/** IFixer which wraps another fixer and makes it validate any assets touched by the fix. */
+struct DATAVALIDATION_API FValidatingFixer : IFixer
+{
+	TSharedPtr<IFixer> Inner;
+
+	virtual EFixApplicability GetApplicability(int32 FixIndex) const override;
+	virtual FFixResult ApplyFix(int32 FixIndex) override;
+
+	static TSharedRef<FValidatingFixer> Create(TSharedRef<IFixer> Inner);
+};
+
+/**
+ * IFixer made out of many smaller fixers. Only one of the fixes in the set may be applied.
+ * Once applied, the remaining fixers become non-applicable.
+ */
+struct DATAVALIDATION_API FMutuallyExclusiveFixSet
+{
+	FMutuallyExclusiveFixSet();
+
+	/** Add a fixer to the set. */
+	void Add(const FText& Label, TSharedRef<IFixer> Inner);
+	/**
+	 * Transform all fixes.
+	 * The callback will be called for each fix and should return a new, transformed fix.
+	 */
+	void Transform(const TFunctionRef<TSharedRef<IFixer>(TSharedRef<IFixer>)>& Callback);
+	/** Generate FFixTokens from the set. These tokens can then be added to a message. */
+	void CreateTokens(const TFunctionRef<void(TSharedRef<FFixToken>)>& Callback) const;
+
+private:
+	struct FSharedData
+	{
+		int32 AppliedFix = INDEX_NONE;
+	};
+	TSharedPtr<FSharedData> SharedData;
+
+	/** Fix belonging to the mutually exclusive fix set. */
+	struct FFixer : IFixer
+	{
+		TSharedPtr<IFixer> Inner;
+		TSharedPtr<FSharedData> SharedData;
+
+		virtual EFixApplicability GetApplicability(int32 FixIndex) const override;
+		virtual FFixResult ApplyFix(int32 FixIndex) override;
+	};
+
+	/**
+	 * Token that is about to be created from the set.
+	 * Note that tokens are not created immediately to allow Transform to work.
+	 */
+	struct FQueuedToken
+	{
+		FText Label;
+		TSharedRef<FFixer> Fixer;
+	};
+
+	TArray<FQueuedToken> QueuedTokens;
+};
+
+}
+
diff --git a/Engine/Source/Developer/MessageLog/Private/UserInterface/SMessageLogMessageListRow.cpp b/Engine/Source/Developer/MessageLog/Private/UserInterface/SMessageLogMessageListRow.cpp
--- a/Engine/Source/Developer/MessageLog/Private/UserInterface/SMessageLogMessageListRow.cpp
+++ b/Engine/Source/Developer/MessageLog/Private/UserInterface/SMessageLogMessageListRow.cpp
@@ -12,6 +12,9 @@
 #include "Internationalization/Regex.h"
 #include "Styling/StyleColors.h"
 #include "AssetRegistry/AssetDataToken.h"
+#include "Framework/Notifications/NotificationManager.h"
+#include "Misc/DataValidation/Fixer.h"
+#include "Widgets/Notifications/SNotificationList.h"
 
 #define LOCTEXT_NAMESPACE "SMessageLogMessageListRow"
 
@@ -52,8 +55,10 @@ TSharedRef<SWidget> SMessageLogMessageListRow::GenerateWidget()
 	// Create the horizontal box and add the icon
 	TSharedRef<SHorizontalBox> MessageBox = SNew(SHorizontalBox);
 	TSharedRef<SHorizontalBox> LinkBox = SNew(SHorizontalBox);
+	TSharedRef<SVerticalBox> FixBox = SNew(SVerticalBox);
 	FName SeverityImageName = NAME_None;
 	bool HasLinks = false;
+	bool bHasFixes = false;
 
 	// Iterate over parts of the message and create widgets for them
 	for (auto TokenIt = MessageTokens.CreateConstIterator(); TokenIt; ++TokenIt)
@@ -80,13 +85,27 @@ TSharedRef<SWidget> SMessageLogMessageListRow::GenerateWidget()
 		}
 			break;
 
+		case EMessageToken::Fix:
+		{
+			TSharedRef<SHorizontalBox> RowWrapper = SNew(SHorizontalBox);
+			CreateMessage(RowWrapper, Token, 2.0f);
+			FixBox->AddSlot()
+			.AutoHeight()
+			.Padding(24.0f, 0.0f, 0.0f, 0.0f)
+			[
+				RowWrapper
+			];
+			bHasFixes = true;
+		}
+			break;
+
 		default:
 			CreateMessage(MessageBox, Token, 2.0f);
 			break;
 		}
 	}
 
-	return SNew(SHorizontalBox)
+	TSharedRef<SHorizontalBox> FullMessageBox = SNew(SHorizontalBox)
 		+ SHorizontalBox::Slot()
 		.AutoWidth()
 		.HAlign(HAlign_Center)
@@ -132,6 +151,26 @@ TSharedRef<SWidget> SMessageLogMessageListRow::GenerateWidget()
 			LinkBox
 		])
 	];
+
+	if (bHasFixes)
+	{
+		return SNew(SVerticalBox)
+		+ SVerticalBox::Slot()
+		.AutoHeight()
+		[
+			FullMessageBox
+		]
+		+ SVerticalBox::Slot()
+		.AutoHeight()
+		[
+			FixBox
+		];
+	}
+	else
+	{
+		return FullMessageBox;
+	}
+
 }
 
 void SMessageLogMessageListRow::CreateMessage(const TSharedRef<SHorizontalBox>& InHorzBox, const TSharedRef<IMessageToken>& InMessageToken, float Padding)
@@ -388,6 +427,64 @@ void SMessageLogMessageListRow::CreateMessage(const TSharedRef<SHorizontalBox>&
 		RowContent = CreateHyperlink(InMessageToken, ActorToken->ToText());
 	}
 		break;
+
+	case EMessageToken::Fix:
+	{
+		TSharedRef<FFixToken> FixToken = StaticCastSharedRef<FFixToken>(InMessageToken);
+
+		RowContent =
+			SNew(SHorizontalBox)
+				+ SHorizontalBox::Slot()
+				.AutoWidth()
+				.Padding(0.0f, 2.0f, 4.0f, 2.0f)
+				[
+					SNew(SImage)
+						.ColorAndOpacity(FSlateColor::UseForeground())
+						.Image_Lambda([FixToken]
+						{
+							switch (FixToken->GetFixer()->GetApplicability(FixToken->GetFixIndex()))
+							{
+								case EFixApplicability::CanBeApplied:
+									return FAppStyle::Get().GetBrush("MessageLog.Fix");
+								case EFixApplicability::Applied:
+									return FAppStyle::Get().GetBrush("Icons.Check");
+								default:
+									return FAppStyle::Get().GetBrush("Icons.BulletPoint16");
+							}
+						})
+				]
+				+ SHorizontalBox::Slot()
+				.AutoWidth()
+				.VAlign(VAlign_Center)
+				[
+					SNew(SHyperlink)
+						.Text(InMessageToken->ToText())
+						.TextStyle(FAppStyle::Get(), "MessageLog")
+						.IsEnabled_Lambda(
+							[FixToken]
+							{ return FixToken->GetFixer()->GetApplicability(FixToken->GetFixIndex()) == EFixApplicability::CanBeApplied; }
+						)
+						.OnNavigate(FSimpleDelegate::CreateSPLambda(
+							AsShared(),
+							[this, FixToken]
+							{
+								FFixResult Result = FixToken->GetFixer()->ApplyFix(FixToken->GetFixIndex());
+								FNotificationInfo NotificationInfo(
+									Result.bIsSuccess ? LOCTEXT("FixAppliedSuccessfully", "Fix applied successfully")
+													  : LOCTEXT("FailedToApplyFix", "Failed to apply fix")
+								);
+								NotificationInfo.SubText = Result.Message;
+								NotificationInfo.bUseThrobber = false;
+								FSlateNotificationManager::Get()
+									.AddNotification(NotificationInfo)
+									->SetCompletionState(
+										Result.bIsSuccess ? SNotificationItem::CS_Success : SNotificationItem::CS_Fail
+									);
+							}
+						))
+				];
+	}
+		break;
 	}
 
 	if (RowContent.IsValid())
diff --git a/Engine/Source/Editor/EditorStyle/Private/StarshipStyle.cpp b/Engine/Source/Editor/EditorStyle/Private/StarshipStyle.cpp
--- a/Engine/Source/Editor/EditorStyle/Private/StarshipStyle.cpp
+++ b/Engine/Source/Editor/EditorStyle/Private/StarshipStyle.cpp
@@ -892,6 +892,7 @@ void FStarshipEditorStyle::FStyle::SetupGeneralStyles()
 		Set( "MessageLog.Docs", new IMAGE_BRUSH( "Icons/icon_Docs_16x", Icon16x16) );
 		Set( "MessageLog.Tutorial", new IMAGE_BRUSH( "Icons/icon_Blueprint_Enum_16x", Icon16x16 ) );
 		Set( "MessageLog.Url", new IMAGE_BRUSH( "Icons/icon_world_16x", Icon16x16 ) );
+		Set( "MessageLog.Fix", new IMAGE_BRUSH_SVG( "Starship/Common/wrench", Icon16x16 ) );
 
 		Set( "MessageLog.TabIcon", new IMAGE_BRUSH_SVG( "Starship/Common/MessageLog", Icon16x16 ) );
 		Set( "MessageLog.ListBorder", new BOX_BRUSH( "/Docking/AppTabContentArea", FMargin(4/16.0f) ) );
diff --git a/Engine/Source/Runtime/Core/Private/Logging/TokenizedMessage.cpp b/Engine/Source/Runtime/Core/Private/Logging/TokenizedMessage.cpp
--- a/Engine/Source/Runtime/Core/Private/Logging/TokenizedMessage.cpp
+++ b/Engine/Source/Runtime/Core/Private/Logging/TokenizedMessage.cpp
@@ -295,4 +295,16 @@ const FOnMessageTokenActivated& FActorToken::GetOnMessageTokenActivated() const
 	}
 }
 
+TSharedRef<FFixToken> FFixToken::Create(const FText& InName, TSharedRef<UE::DataValidation::IFixer> InFixer, int32 InFixIndex)
+{
+	return MakeShared<FFixToken>(FPrivateToken(), InName, MoveTemp(InFixer), InFixIndex);
+}
+
+FFixToken::FFixToken(const FText& InName, TSharedPtr<UE::DataValidation::IFixer> InFixer, int32 InFixIndex)
+	: Fixer(MoveTemp(InFixer))
+	, FixIndex(InFixIndex)
+{
+	CachedText = InName.IsEmpty() ? LOCTEXT("FixLabelEmpty", "Fix Issue") : InName;
+}
+
 #undef LOCTEXT_NAMESPACE
diff --git a/Engine/Source/Runtime/Core/Public/Logging/TokenizedMessage.h b/Engine/Source/Runtime/Core/Public/Logging/TokenizedMessage.h
--- a/Engine/Source/Runtime/Core/Public/Logging/TokenizedMessage.h
+++ b/Engine/Source/Runtime/Core/Public/Logging/TokenizedMessage.h
@@ -47,6 +47,7 @@ namespace EMessageToken
 		URL,
 		EdGraph,
 		DynamicText,
+		Fix,
 	};
 }
 
@@ -749,4 +750,77 @@ private:
 
 	/** The default activation method, if any */
 	CORE_API static FOnMessageTokenActivated DefaultMessageTokenActivated;
-};
\ No newline at end of file
+};
+
+namespace UE::DataValidation
+{
+struct IFixer;
+}
+
+struct FFixResult
+{
+	bool bIsSuccess = true;
+	FText Message;
+
+	static FFixResult Success(const FText& Message)
+	{
+		return {true, Message};
+	}
+
+	static FFixResult Success()
+	{
+		return {true, FText::GetEmpty()};
+	}
+
+	static FFixResult Failure(const FText& Message)
+	{
+		return {false, Message};
+	}
+
+	// NOTE: No overload for Failure that accepts no message, because in case of an error you
+	// should *always* let the user know what went wrong.
+};
+
+enum class EFixApplicability : uint8 
+{
+	/** Fix hasn't been applied yet but can be applied */
+	CanBeApplied,
+	/** Fix has been applied successfully */
+	Applied,
+	/** Fix was not applied by the user, and instead a different fix was chosen */
+	DidNotApply,
+};
+
+/** A fix action token. These are displayed in a separate list below the log message. */
+class FFixToken : public IMessageToken
+{
+	// The private token allows only members or friends to call MakeShared.
+	struct FPrivateToken { explicit FPrivateToken() = default; };
+
+public:
+	/** Creates an FFixToken out of an IFixer and a fix index. Prefer IFixer::CreateToken over this. */
+	CORE_API static TSharedRef<FFixToken> Create(const FText& InName, TSharedRef<UE::DataValidation::IFixer> InFixer, int32 InFixIndex);
+
+	virtual EMessageToken::Type GetType() const override
+	{
+		return EMessageToken::Fix;
+	}
+
+	TSharedRef<UE::DataValidation::IFixer> GetFixer() const { return Fixer.ToSharedRef(); }
+	int32 GetFixIndex() const { return FixIndex; }
+	void SetFix(TSharedRef<UE::DataValidation::IFixer> InFixer, int32 InFixIndex)
+	{
+		Fixer = InFixer;
+		FixIndex = InFixIndex;
+	}
+
+	FFixToken(FPrivateToken, const FText& InName, TSharedPtr<UE::DataValidation::IFixer> Fixer, int32 FixIndex)
+		: FFixToken(InName, MoveTemp(Fixer), FixIndex)
+	{}
+
+private:
+	CORE_API FFixToken(const FText& InName, TSharedPtr<UE::DataValidation::IFixer> InFixer, int32 InFixIndex);
+
+	TSharedPtr<UE::DataValidation::IFixer> Fixer;
+	int32 FixIndex = INDEX_NONE;
+};
diff --git a/Engine/Source/Runtime/CoreUObject/Private/Misc/DataValidation/Fixer.cpp b/Engine/Source/Runtime/CoreUObject/Private/Misc/DataValidation/Fixer.cpp
new file mode 100644
--- /dev/null
+++ b/Engine/Source/Runtime/CoreUObject/Private/Misc/DataValidation/Fixer.cpp
@@ -0,0 +1,19 @@
+﻿// Copyright Epic Games, Inc. All Rights Reserved.
+
+#include "Misc/DataValidation/Fixer.h"
+
+#include "Internationalization/Internationalization.h"
+
+#define LOCTEXT_NAMESPACE "DataValidationFixers"
+
+namespace UE::DataValidation
+{
+
+TSharedRef<FFixToken> IFixer::CreateToken(const FText& Label)
+{
+	return FFixToken::Create(FText::Format(LOCTEXT("SingleFix", "Fix: {0}"), Label), AsShared(), 0);
+}
+
+}
+
+#undef LOCTEXT_NAMESPACE
diff --git a/Engine/Source/Runtime/CoreUObject/Public/Misc/DataValidation/Fixer.h b/Engine/Source/Runtime/CoreUObject/Public/Misc/DataValidation/Fixer.h
new file mode 100644
--- /dev/null
+++ b/Engine/Source/Runtime/CoreUObject/Public/Misc/DataValidation/Fixer.h
@@ -0,0 +1,45 @@
+﻿// Copyright Epic Games, Inc. All Rights Reserved.
+
+#pragma once
+
+#include "Logging/TokenizedMessage.h"
+#include "Templates/SharedPointer.h"
+
+namespace UE::DataValidation
+{
+
+/**
+ * Provider of automatic fixes for an asset.
+ *
+ * Fixers are composable: it is possible to layer many fixers on top of each other to provide
+ * fixes which perform extra actions on top of a base fixer.
+ * These layers are provided within the DataValidation plugin. See DataValidationFixers.h.
+ */
+struct COREUOBJECT_API IFixer : TSharedFromThis<IFixer>
+{
+	virtual ~IFixer() = default;
+
+	/** Returns whether the fix can be applied at the moment. This must be called before `ApplyFix`. */
+	virtual EFixApplicability GetApplicability(int32 FixIndex) const = 0;
+	/** Applies a fix provided by this fixer. */
+	virtual FFixResult ApplyFix(int32 FixIndex) = 0;
+
+	/**
+	 * Sugar for wrapping the fixer in another fixer using postfix `Fixer->WrappedIn<FOtherFixer>(Args)`
+	 * instead of `FSomeFixer::Create(Fixer, Args)`.
+	 */
+	template <typename FixerType, typename... ArgTypes>
+	TSharedRef<FixerType> WrappedIn(ArgTypes&&... Args)
+	{
+		return FixerType::Create(AsShared(), Forward<ArgTypes>(Args)...);
+	}
+
+	/**
+	 * Creates an FFixToken out of the fix and a given label.
+	 * This should generally be preferred over the lower level FFixToken::Create.
+	 */
+	TSharedRef<FFixToken> CreateToken(const FText& Label);
+};
+
+}
+
diff --git a/Engine/Source/Runtime/SlateCore/Private/Styling/StarshipCoreStyle.cpp b/Engine/Source/Runtime/SlateCore/Private/Styling/StarshipCoreStyle.cpp
--- a/Engine/Source/Runtime/SlateCore/Private/Styling/StarshipCoreStyle.cpp
+++ b/Engine/Source/Runtime/SlateCore/Private/Styling/StarshipCoreStyle.cpp
@@ -394,6 +394,7 @@ TSharedRef<ISlateStyle> FStarshipCoreStyle::Create()
 		Style->Set("Icons.Unlink", new IMAGE_BRUSH_SVG("Starship/Common/Unlinked", Icon16x16));
 
 		Style->Set("Icons.BulletPoint", new IMAGE_BRUSH_SVG("Starship/Common/bullet-point", Icon8x8));
+		Style->Set("Icons.BulletPoint16", new IMAGE_BRUSH_SVG("Starship/Common/bullet-point16", Icon16x16));
 
 		Style->Set("Icons.SortDown", new IMAGE_BRUSH_SVG("Starship/Common/SortDown", Icon16x16));
 		Style->Set("Icons.SortUp", new IMAGE_BRUSH_SVG("Starship/Common/SortUp", Icon16x16));
