Added utility to iterate the enabled validators

Wraps the load + iterate logic into a single place

#rnx

diff --git a/Engine/Plugins/Editor/DataValidation/Source/DataValidation/Private/EditorValidatorSubsystem.cpp b/Engine/Plugins/Editor/DataValidation/Source/DataValidation/Private/EditorValidatorSubsystem.cpp
--- a/Engine/Plugins/Editor/DataValidation/Source/DataValidation/Private/EditorValidatorSubsystem.cpp
+++ b/Engine/Plugins/Editor/DataValidation/Source/DataValidation/Private/EditorValidatorSubsystem.cpp
@@ -222,6 +222,23 @@ void UEditorValidatorSubsystem::CleanupValidators()
 	Validators.Empty();
 }
 
+void UEditorValidatorSubsystem::ForEachEnabledValidator(TFunctionRef<bool(UEditorValidatorBase* Validator)> Callback) const
+{
+	LoadValidators();
+
+	for (const TPair<FTopLevelAssetPath, TObjectPtr<UEditorValidatorBase>>& ValidatorPair : Validators)
+	{
+		if (UEditorValidatorBase* Validator = ValidatorPair.Value;
+			Validator && Validator->IsEnabled())
+		{
+			if (!Callback(Validator))
+			{
+				break;
+			}
+		}
+	}
+}
+
 EDataValidationResult UEditorValidatorSubsystem::IsObjectValid(
 	UObject* InObject,
 	TArray<FText>& ValidationErrors,
@@ -313,19 +330,13 @@ EDataValidationResult UEditorValidatorSubsystem::ValidateObjectInternal(
 		{
 			return Result;
 		}
-			
-		// This doesn't change the logical view of the subsystem
-		const_cast<UEditorValidatorSubsystem&>(*this).LoadValidators();
-
-		for (const TPair<FTopLevelAssetPath, TObjectPtr<UEditorValidatorBase>>& ValidatorPair : Validators)
+		
+		ForEachEnabledValidator([InObject, &InAssetData, &InContext, &Result](UEditorValidatorBase* Validator)
 		{
-			UEditorValidatorBase* Validator = ValidatorPair.Value;
-			if (Validator && Validator->IsEnabled())
-			{
-				EDataValidationResult NewResult = Validator->ValidateLoadedAsset(InAssetData, InObject, InContext);
-				Result = CombineDataValidationResults(Result, NewResult);
-			}
-		}
+			EDataValidationResult NewResult = Validator->ValidateLoadedAsset(InAssetData, InObject, InContext);
+			Result = CombineDataValidationResults(Result, NewResult);
+			return true;
+		});
 	}
 
 	return Result;
@@ -358,13 +369,13 @@ EDataValidationResult UEditorValidatorSubsystem::ValidateAssetsInternal(
 	FScopedSlowTask SlowTask(AssetDataList.Num(), LOCTEXT("DataValidation.ValidateAssetsTask", "Validating Assets"));
 	SlowTask.MakeDialog();
 	
-	LoadValidators();
 	UE_LOG(LogContentValidation, Display, TEXT("Starting to validate %d assets"), AssetDataList.Num());
-	UE_LOG(LogContentValidation, Log, TEXT("Registered validators:"));
-	for (const TPair<FTopLevelAssetPath, TObjectPtr<UEditorValidatorBase>>& ValidatorPair : Validators)
+	UE_LOG(LogContentValidation, Log, TEXT("Enabled validators:"));
+	ForEachEnabledValidator([](UEditorValidatorBase* Validator)
 	{
-		UE_LOG(LogContentValidation, Log, TEXT("\t%s"), *ValidatorPair.Key.ToString());
-	}
+		UE_LOG(LogContentValidation, Log, TEXT("\t%s"), *Validator->GetClass()->GetClassPathName().ToString());
+		return true;
+	});
 	
 	// Broadcast the Editor event before we start validating. This lets other systems (such as Sequencer) restore the state
 	// of the level to what is actually saved on disk before performing validation.
@@ -888,8 +899,6 @@ void UEditorValidatorSubsystem::GatherAssetsToValidateFromChangelist(
 {
 	IAssetRegistry& AssetRegistry = IAssetRegistry::GetChecked();
 	
-	LoadValidators();
-	
 	for (const FName& PackageName : InChangelist->ModifiedPackageNames)
 	{
 		TArray<FAssetData> NewAssets;
@@ -898,24 +907,20 @@ void UEditorValidatorSubsystem::GatherAssetsToValidateFromChangelist(
 	}
 	
 	// Gather assets requested by plugin/project validators 
-	for (const TPair<FTopLevelAssetPath, TObjectPtr<UEditorValidatorBase>>& ValidatorPair : Validators)
+	ForEachEnabledValidator([this, InChangelist, &Settings, &InContext, &OutAssets](UEditorValidatorBase* Validator)
 	{
-		UEditorValidatorBase* Validator = ValidatorPair.Value;
-		if (Validator && Validator->IsEnabled())
+		TArray<FAssetData> NewAssets = Validator->GetAssetsToValidateFromChangelist(InChangelist, InContext);
+		for (const FAssetData& Asset : NewAssets)
 		{
-			TArray<FAssetData> NewAssets = Validator->GetAssetsToValidateFromChangelist(InChangelist, InContext);
-			for (const FAssetData& Asset : NewAssets)
+			// It's not strictly necessary to filter assets here but it makes logging simpler
+			if (ShouldValidateAsset(Asset, Settings, InContext))
 			{
-				// It's not strictly necessary to filter assets here but it makes logging simpler
-				if (ShouldValidateAsset(Asset, Settings, InContext))
-				{
-					UE_LOG(LogContentValidation, Verbose, TEXT("Asset validator %s adding %s to be validated."), *Validator->GetPathName(), *Asset.GetSoftObjectPath().ToString());
-					OutAssets.Add(Asset);
-				}
+				UE_LOG(LogContentValidation, Verbose, TEXT("Asset validator %s adding %s to be validated."), *Validator->GetPathName(), *Asset.GetSoftObjectPath().ToString());
+				OutAssets.Add(Asset);
 			}
 		}
-	}
-
+		return true;
+	});
 	
 	if (Settings.bValidateReferencersOfDeletedAssets)
 	{
diff --git a/Engine/Plugins/Editor/DataValidation/Source/DataValidation/Public/EditorValidatorSubsystem.h b/Engine/Plugins/Editor/DataValidation/Source/DataValidation/Public/EditorValidatorSubsystem.h
--- a/Engine/Plugins/Editor/DataValidation/Source/DataValidation/Public/EditorValidatorSubsystem.h
+++ b/Engine/Plugins/Editor/DataValidation/Source/DataValidation/Public/EditorValidatorSubsystem.h
@@ -233,6 +233,12 @@ public:
 	UFUNCTION(BlueprintCallable, Category = "Validation")
 	void AddValidator(UEditorValidatorBase* InValidator);
 
+	/**
+	 * Iterate the enabled set of validators.
+	 * @note Return true to continue iteration, or false to stop.
+	 */
+	void ForEachEnabledValidator(TFunctionRef<bool(UEditorValidatorBase* Validator)> Callback) const;
+
 	/**
 	 * Runs registered validators on the provided object.
 	 * Does not add anything to any FMessageLog tabs.
