Log the enabled set of validators when running cook validation

#rnx

diff --git a/Engine/Plugins/Editor/DataValidation/Source/DataValidation/Private/DataValidationModule.cpp b/Engine/Plugins/Editor/DataValidation/Source/DataValidation/Private/DataValidationModule.cpp
--- a/Engine/Plugins/Editor/DataValidation/Source/DataValidation/Private/DataValidationModule.cpp
+++ b/Engine/Plugins/Editor/DataValidation/Source/DataValidation/Private/DataValidationModule.cpp
@@ -12,6 +12,7 @@
 #include "AssetRegistry/AssetRegistryModule.h"
 #include "Misc/MessageDialog.h"
 #include "DataValidationCommandlet.h"
+#include "EditorValidatorBase.h"
 #include "Elements/Framework/TypedElementSelectionSet.h"
 #include "Logging/MessageLog.h"
 #include "UObject/ICookInfo.h"
@@ -370,6 +371,20 @@ EDataValidationResult FDataValidationModule::OnValidateSourcePackageDuringCook(U
 
 	if (UEditorValidatorSubsystem* EditorValidationSubsystem = GEditor->GetEditorSubsystem<UEditorValidatorSubsystem>())
 	{
+		// Log the enabled set of validators, but only once
+		{
+			static const bool LogValidatorsListOnce = [EditorValidationSubsystem]()
+			{
+				UE_LOG(LogContentValidation, Log, TEXT("Enabled validators:"));
+				EditorValidationSubsystem->ForEachEnabledValidator([](UEditorValidatorBase* Validator)
+				{
+					UE_LOG(LogContentValidation, Log, TEXT("\t%s"), *Validator->GetClass()->GetClassPathName().ToString());
+					return true;
+				});
+				return true;
+			}();
+		}
+
 		TArray<FAssetData> AssetList;
 		IAssetRegistry::GetChecked().GetAssetsByPackageName(Package->GetFName(), AssetList);
 
