Add links to validator errors

diff --git a/Engine/Plugins/Editor/DataValidation/Source/DataValidation/Private/EditorValidatorBase.cpp b/Engine/Plugins/Editor/DataValidation/Source/DataValidation/Private/EditorValidatorBase.cpp
--- a/Engine/Plugins/Editor/DataValidation/Source/DataValidation/Private/EditorValidatorBase.cpp
+++ b/Engine/Plugins/Editor/DataValidation/Source/DataValidation/Private/EditorValidatorBase.cpp
@@ -118,6 +118,17 @@ void UEditorValidatorBase::AssetFails(const UObject* InAsset, const FText& InMes
 	ValidationResult = EDataValidationResult::Invalid;
 }
 
+void UEditorValidatorBase::AssetFailsTokenized(const UObject* InAsset, const FText& InMessage)
+{
+	FFormatNamedArguments Arguments;
+	Arguments.Add(TEXT("CustomMessage"), InMessage);
+	Arguments.Add(TEXT("ValidatorName"), FText::FromString(GetClass()->GetName()));
+
+	FText FailureMessage = FText::Format(LOCTEXT("AssetCheck_Message_Error", "{CustomMessage}. ({ValidatorName})"), Arguments);
+	AssetMessage(InAsset, EMessageSeverity::Error, FailureMessage);
+	ValidationResult = EDataValidationResult::Invalid;
+}
+
 void UEditorValidatorBase::AssetWarning(const UObject* InAsset, const FText& InMessage)
 {
 	FFormatNamedArguments Arguments;
@@ -128,6 +139,16 @@ void UEditorValidatorBase::AssetWarning(const UObject* InAsset, const FText& InM
 	AllWarnings.Add(WarningMessage);
 }
 
+void UEditorValidatorBase::AssetWarningTokenized(const UObject* InAsset, const FText& InMessage)
+{
+	FFormatNamedArguments Arguments;
+	Arguments.Add(TEXT("CustomMessage"), InMessage);
+	Arguments.Add(TEXT("ValidatorName"), FText::FromString(GetClass()->GetName()));
+
+	FText WarningMessage = FText::Format(LOCTEXT("AssetCheck_Message_Warning", "{CustomMessage}. ({ValidatorName})"), Arguments);
+	AssetMessage(InAsset, EMessageSeverity::Warning, WarningMessage);
+}
+
 void UEditorValidatorBase::AssetPasses(const UObject* InAsset)
 {
 	if (LogContentValidation.GetVerbosity() >= ELogVerbosity::Verbose)
diff --git a/Engine/Plugins/Editor/DataValidation/Source/DataValidation/Public/EditorValidatorBase.h b/Engine/Plugins/Editor/DataValidation/Source/DataValidation/Public/EditorValidatorBase.h
--- a/Engine/Plugins/Editor/DataValidation/Source/DataValidation/Public/EditorValidatorBase.h
+++ b/Engine/Plugins/Editor/DataValidation/Source/DataValidation/Public/EditorValidatorBase.h
@@ -75,6 +75,10 @@ public:
 	UFUNCTION(BlueprintCallable, Category = "Asset Validation")
 	void AssetFails(const UObject* InAsset, const FText& InMessage);
 
+	/** Marks the validation as failed and adds an error message that links to the specified asset. */
+	UFUNCTION(BlueprintCallable, Category = "Asset Validation")
+	void AssetFailsTokenized(const UObject* InAsset, const FText& InMessage);
+
 	UE_DEPRECATED("5.4", "This function has been replaced by the version AssetFails(UObject*, const FText&)")
 	void AssetFails(const UObject* InAsset, const FText& InMessage, TArray<FText>& InOutErrors);
 
@@ -85,6 +89,10 @@ public:
 	/** Adds a message to this validation but doesn't mark it as failed. */
 	UFUNCTION(BlueprintCallable, Category = "Asset Validation")
 	void AssetWarning(const UObject* InAsset, const FText& InMessage);
+
+	/** Adds a message that links to the specified asset to this validation but doesn't mark it as failed. */
+	UFUNCTION(BlueprintCallable, Category = "Asset Validation")
+	void AssetWarningTokenized(const UObject* InAsset, const FText& InMessage);
 	
 	/** 
 	 * Add a tokenized message to the validation results. If the severity is error, marks the validation as failed.
