Log enabled validators during cook as Display rather than Log

This makes them easier to see in CIS logs

#rnx

diff --git a/Engine/Plugins/Editor/DataValidation/Source/DataValidation/Private/DataValidationModule.cpp b/Engine/Plugins/Editor/DataValidation/Source/DataValidation/Private/DataValidationModule.cpp
--- a/Engine/Plugins/Editor/DataValidation/Source/DataValidation/Private/DataValidationModule.cpp
+++ b/Engine/Plugins/Editor/DataValidation/Source/DataValidation/Private/DataValidationModule.cpp
@@ -375,10 +375,10 @@ EDataValidationResult FDataValidationModule::OnValidateSourcePackageDuringCook(U
 		{
 			static const bool LogValidatorsListOnce = [EditorValidationSubsystem]()
 			{
-				UE_LOG(LogContentValidation, Log, TEXT("Enabled validators:"));
+				UE_LOG(LogContentValidation, Display, TEXT("Enabled validators:"));
 				EditorValidationSubsystem->ForEachEnabledValidator([](UEditorValidatorBase* Validator)
 				{
-					UE_LOG(LogContentValidation, Log, TEXT("\t%s"), *Validator->GetClass()->GetClassPathName().ToString());
+					UE_LOG(LogContentValidation, Display, TEXT("\t%s"), *Validator->GetClass()->GetClassPathName().ToString());
 					return true;
 				});
 				return true;
