Mark asset validator loads as editor-only for the cooker

This avoids them being detected as an asset reference that needs to be cooked

#rnx

diff --git a/Engine/Plugins/Editor/DataValidation/Source/DataValidation/Private/EditorValidatorSubsystem.cpp b/Engine/Plugins/Editor/DataValidation/Source/DataValidation/Private/EditorValidatorSubsystem.cpp
--- a/Engine/Plugins/Editor/DataValidation/Source/DataValidation/Private/EditorValidatorSubsystem.cpp
+++ b/Engine/Plugins/Editor/DataValidation/Source/DataValidation/Private/EditorValidatorSubsystem.cpp
@@ -26,6 +26,7 @@
 #include "Misc/PackageName.h"
 #include "Misc/ScopedSlowTask.h"
 #include "Misc/ScopeExit.h"
+#include "UObject/ICookInfo.h"
 #include "UObject/TopLevelAssetPath.h"
 
 #include UE_INLINE_GENERATED_CPP_BY_NAME(EditorValidatorSubsystem)
@@ -962,6 +963,7 @@ void UEditorValidatorSubsystem::LoadValidators()
 				// If this object isn't currently loaded, load it
 				if (ValidatorObject == nullptr)
 				{
+					FCookLoadScope EditorOnlyLoadScope(ECookLoadType::EditorOnly);
 					FSoftObjectPathSerializationScope SerializationScope(NAME_None, NAME_None, ESoftObjectPathCollectType::EditorOnlyCollect, ESoftObjectPathSerializeType::AlwaysSerialize);
 					ValidatorObject = ValidatorObjectSoftPath.TryLoad();
 				}
