World Partition: fix missing actor desc container instances uninitialize calls.

diff --git a/Engine/Plugins/Editor/DataValidation/Source/DataValidation/Private/WorldPartitionChangelistValidator.cpp b/Engine/Plugins/Editor/DataValidation/Source/DataValidation/Private/WorldPartitionChangelistValidator.cpp
--- a/Engine/Plugins/Editor/DataValidation/Source/DataValidation/Private/WorldPartitionChangelistValidator.cpp
+++ b/Engine/Plugins/Editor/DataValidation/Source/DataValidation/Private/WorldPartitionChangelistValidator.cpp
@@ -137,7 +137,7 @@ void UWorldPartitionChangelistValidator::ValidateActorsAndDataLayersFromChangeLi
 		}
 	}
 
-	auto RegisterContainerToValidate = [](UWorld* InWorld, FName InContainerPackageName, FActorDescContainerInstanceCollection& OutRegisteredContainers, const FGuid& InContentBundleGuid = FGuid(), const UExternalDataLayerAsset* InExternalDataLayerAsset = nullptr)
+	auto RegisterContainerToValidate = [](UWorld* InWorld, FName InContainerPackageName, FActorDescContainerInstanceCollection& OutRegisteredContainers, TArray<UActorDescContainerInstance*>& OutNewlyCreatedContainers, const FGuid& InContentBundleGuid = FGuid(), const UExternalDataLayerAsset* InExternalDataLayerAsset = nullptr)
 	{
 		if (OutRegisteredContainers.Contains(InContainerPackageName))
 		{
@@ -162,6 +162,7 @@ void UWorldPartitionChangelistValidator::ValidateActorsAndDataLayersFromChangeLi
 			InitializeParams.ContentBundleGuid = InContentBundleGuid;
 			InitializeParams.ExternalDataLayerAsset = InExternalDataLayerAsset;
 			ContainerInstance->Initialize(InitializeParams);
+			OutNewlyCreatedContainers.Add(ContainerInstance);
 		}
 		else
 		{
@@ -184,9 +185,10 @@ void UWorldPartitionChangelistValidator::ValidateActorsAndDataLayersFromChangeLi
 		TGuardValue<UObject*> GuardCurrentAsset(CurrentAsset, World);
 		
 		FActorDescContainerInstanceCollection ContainersToValidate;
+		TArray<UActorDescContainerInstance*> ContainersToUninit;
 
 		// Always register the main world container because content bundle containers can't be validated separately
-		RegisterContainerToValidate(World, MapPath.GetPackageName(), ContainersToValidate);
+		RegisterContainerToValidate(World, MapPath.GetPackageName(), ContainersToValidate, ContainersToUninit);
 
 		TSet<FSoftObjectPath> ProcessedExternalDataLayersForMap;
 		for (const FAssetData& ActorData : ActorsData)
@@ -200,7 +202,7 @@ void UWorldPartitionChangelistValidator::ValidateActorsAndDataLayersFromChangeLi
 				FString ContentBundleContainerPackagePath;
 				verify(ContentBundlePaths::BuildActorDescContainerPackagePath(FString(ContentBundleMountPoint), ContentBundleGuid, MapPath.GetPackageName().ToString(), ContentBundleContainerPackagePath));
 
-				RegisterContainerToValidate(World, FName(*ContentBundleContainerPackagePath), ContainersToValidate, ContentBundleGuid);
+				RegisterContainerToValidate(World, FName(*ContentBundleContainerPackagePath), ContainersToValidate, ContainersToUninit, ContentBundleGuid);
 			}
 			else
 			{
@@ -216,7 +218,7 @@ void UWorldPartitionChangelistValidator::ValidateActorsAndDataLayersFromChangeLi
 							if (const UExternalDataLayerAsset* ExternalDataLayerAsset = Cast<UExternalDataLayerAsset>(ExternalDataLayerPath.TryLoad()))
 							{
 								const FString EDLContainerPackagePath = FExternalDataLayerHelper::GetExternalDataLayerLevelRootPath(ExternalDataLayerAsset, MapPath.GetPackageName().ToString());
-								RegisterContainerToValidate(World, FName(*EDLContainerPackagePath), ContainersToValidate, FGuid(), ExternalDataLayerAsset);
+								RegisterContainerToValidate(World, FName(*EDLContainerPackagePath), ContainersToValidate, ContainersToUninit, FGuid(), ExternalDataLayerAsset);
 							}
 						}
 					}
@@ -244,6 +246,11 @@ void UWorldPartitionChangelistValidator::ValidateActorsAndDataLayersFromChangeLi
 
 		Params.ActorDescContainerInstanceCollection = &ContainersToValidate;
 		UWorldPartition::CheckForErrors(Params);
+
+		for (UActorDescContainerInstance* ContainerToUninit : ContainersToUninit)
+		{
+			ContainerToUninit->Uninitialize();
+		}
 	}
 }
 
diff --git a/Engine/Source/Runtime/Engine/Private/WorldPartition/ActorDescContainerInstance.cpp b/Engine/Source/Runtime/Engine/Private/WorldPartition/ActorDescContainerInstance.cpp
--- a/Engine/Source/Runtime/Engine/Private/WorldPartition/ActorDescContainerInstance.cpp
+++ b/Engine/Source/Runtime/Engine/Private/WorldPartition/ActorDescContainerInstance.cpp
@@ -16,16 +16,6 @@
 #if WITH_EDITOR
 UActorDescContainerInstance::FActorDescContainerInstanceInitializeDelegate UActorDescContainerInstance::OnActorDescContainerInstanceInitialized;
 
-void UActorDescContainerInstance::BeginDestroy()
-{
-	Super::BeginDestroy();
-
-	if (bIsInitialized)
-	{
-		Uninitialize();
-	}
-}
-
 FName UActorDescContainerInstance::GetContainerPackageNameFromWorld(UWorld* InWorld)
 {
 	check(InWorld);
diff --git a/Engine/Source/Runtime/Engine/Public/WorldPartition/ActorDescContainerInstance.h b/Engine/Source/Runtime/Engine/Public/WorldPartition/ActorDescContainerInstance.h
--- a/Engine/Source/Runtime/Engine/Public/WorldPartition/ActorDescContainerInstance.h
+++ b/Engine/Source/Runtime/Engine/Public/WorldPartition/ActorDescContainerInstance.h
@@ -32,10 +32,6 @@ protected:
 	{}
 
 #if WITH_EDITOR
-	//~ Begin UObject Interface
-	virtual void BeginDestroy() override;
-	//~ End UObject Interface
-
 	friend FWorldPartitionActorDescInstance;
 	friend UWorldPartition;
 		
