Add new function so that validators can be removed.

diff --git a/Engine/Plugins/Editor/DataValidation/Source/DataValidation/Private/EditorValidatorSubsystem.cpp b/Engine/Plugins/Editor/DataValidation/Source/DataValidation/Private/EditorValidatorSubsystem.cpp
--- a/Engine/Plugins/Editor/DataValidation/Source/DataValidation/Private/EditorValidatorSubsystem.cpp
+++ b/Engine/Plugins/Editor/DataValidation/Source/DataValidation/Private/EditorValidatorSubsystem.cpp
@@ -218,6 +218,25 @@ void UEditorValidatorSubsystem::AddValidator(const FAssetData& InValidatorAssetD
 	}
 }
 
+void UEditorValidatorSubsystem::RemoveValidator(UEditorValidatorBase* InValidator)
+{
+	if (InValidator)
+	{
+		UClass* Class = InValidator->GetClass();
+		if (Cast<UBlueprintGeneratedClass>(Class))
+		{
+			if (UObject* ClassGenerator = Class->ClassGeneratedBy)
+			{
+				Validators.Remove(FTopLevelAssetPath(ClassGenerator));
+			}
+		}
+		else
+		{
+			Validators.Remove(InValidator->GetClass()->GetClassPathName());
+		}
+	}
+}
+
 void UEditorValidatorSubsystem::CleanupValidators()
 {
 	Validators.Empty();
diff --git a/Engine/Plugins/Editor/DataValidation/Source/DataValidation/Public/EditorValidatorSubsystem.h b/Engine/Plugins/Editor/DataValidation/Source/DataValidation/Public/EditorValidatorSubsystem.h
--- a/Engine/Plugins/Editor/DataValidation/Source/DataValidation/Public/EditorValidatorSubsystem.h
+++ b/Engine/Plugins/Editor/DataValidation/Source/DataValidation/Public/EditorValidatorSubsystem.h
@@ -233,6 +233,13 @@ public:
 	UFUNCTION(BlueprintCallable, Category = "Validation")
 	void AddValidator(UEditorValidatorBase* InValidator);
 
+	/*
+	* Removes a validator
+	* Should be called during module shutdown if a validator was added.
+	*/
+	UFUNCTION(BlueprintCallable, Category = "Validation")
+	void RemoveValidator(UEditorValidatorBase* InValidator);
+
 	/**
 	 * Iterate the enabled set of validators.
 	 * @note Return true to continue iteration, or false to stop.
