Additional logging

diff --git a/Engine/Plugins/Editor/DataValidation/Source/DataValidation/Private/EditorValidatorSubsystem.cpp b/Engine/Plugins/Editor/DataValidation/Source/DataValidation/Private/EditorValidatorSubsystem.cpp
--- a/Engine/Plugins/Editor/DataValidation/Source/DataValidation/Private/EditorValidatorSubsystem.cpp
+++ b/Engine/Plugins/Editor/DataValidation/Source/DataValidation/Private/EditorValidatorSubsystem.cpp
@@ -430,6 +430,11 @@ EDataValidationResult UEditorValidatorSubsystem::ValidateAssetsInternal(
 			{
 				AssetDataList.Add(Pair.Key);
 			}
+			else
+			{
+				UE_LOG(LogContentValidation, Display, TEXT("Package %s (owner of some external objects) being skipped for validation."), 
+					*WriteToString<256>(Pair.Key.PackageName));
+			}
 		}
 		UE::DataValidation::AddAssetValidationMessages(DataValidationLog, ValidationContext);
 		DataValidationLog.Flush();
@@ -896,7 +901,7 @@ EDataValidationResult UEditorValidatorSubsystem::ValidateChangelistsInternal(
 		{
 			if (!ShouldValidateAsset(*It, Settings, ValidationContext))
 			{
-				UE_LOG(LogContentValidation, Log, TEXT("Excluding asset %s from validation"), *It->GetSoftObjectPath().ToString());
+				UE_LOG(LogContentValidation, Display, TEXT("Excluding asset %s from validation"), *It->GetSoftObjectPath().ToString());
 				It.RemoveCurrent();
 			}
 		}
@@ -936,7 +941,7 @@ void UEditorValidatorSubsystem::GatherAssetsToValidateFromChangelist(
 			// It's not strictly necessary to filter assets here but it makes logging simpler
 			if (ShouldValidateAsset(Asset, Settings, InContext))
 			{
-				UE_LOG(LogContentValidation, Verbose, TEXT("Asset validator %s adding %s to be validated."), *Validator->GetPathName(), *Asset.GetSoftObjectPath().ToString());
+				UE_LOG(LogContentValidation, Display, TEXT("Asset validator %s adding %s to be validated from changelist %s."), *Validator->GetPathName(), *Asset.GetSoftObjectPath().ToString(), *InChangelist->Description.ToString());
 				OutAssets.Add(Asset);
 			}
 		}
@@ -955,7 +960,7 @@ void UEditorValidatorSubsystem::GatherAssetsToValidateFromChangelist(
 			AssetRegistry.GetReferencers(DeletedPackageName, PackageReferencers, UE::AssetRegistry::EDependencyCategory::Package);
 			for (const FName& Referencer : PackageReferencers)
 			{
-				UE_LOG(LogContentValidation, Log, TEXT("Adding %s to to validated as it is a referencer of deleted asset %s"), *Referencer.ToString(), *DeletedPackageName.ToString());
+				UE_LOG(LogContentValidation, Display, TEXT("Adding %s to to validated as it is a referencer of deleted asset %s"), *Referencer.ToString(), *DeletedPackageName.ToString());
 				TArray<FAssetData> NewAssets;
 				AssetRegistry.GetAssetsByPackageName(Referencer, NewAssets, true);
 				OutAssets.Append(NewAssets);
