Remove log line that is duplicated by use of MessageLog. To get these message logs at the Display log level, use MessageLog.InfoMessagesAtDisplayVerbosity

diff --git a/Engine/Plugins/Editor/DataValidation/Source/DataValidation/Private/EditorValidatorSubsystem.cpp b/Engine/Plugins/Editor/DataValidation/Source/DataValidation/Private/EditorValidatorSubsystem.cpp
--- a/Engine/Plugins/Editor/DataValidation/Source/DataValidation/Private/EditorValidatorSubsystem.cpp
+++ b/Engine/Plugins/Editor/DataValidation/Source/DataValidation/Private/EditorValidatorSubsystem.cpp
@@ -480,12 +480,11 @@ EDataValidationResult UEditorValidatorSubsystem::ValidateAssetsInternal(
 			++OutResults.NumSkipped;
 			continue;
 		}
-		
+
 		DataValidationLog.Info()
 			->AddToken(FAssetDataToken::Create(Data))
 			->AddToken(FTextToken::Create(LOCTEXT("Data.ValidatingAsset", "Validating asset")));
-		UE_LOG(LogContentValidation, Display, TEXT("Validating asset %s"), *Data.ToSoftObjectPath().ToString());
-		
+
 		UObject* LoadedAsset = Data.FastGetAsset(false);
 		const bool bAlreadyLoaded = LoadedAsset != nullptr;
 
