Make attempts to validate changelists while asset registry scan is in progress fatal in unattended mode.

diff --git a/Engine/Plugins/Editor/DataValidation/Source/DataValidation/Private/EditorValidatorSubsystem.cpp b/Engine/Plugins/Editor/DataValidation/Source/DataValidation/Private/EditorValidatorSubsystem.cpp
--- a/Engine/Plugins/Editor/DataValidation/Source/DataValidation/Private/EditorValidatorSubsystem.cpp
+++ b/Engine/Plugins/Editor/DataValidation/Source/DataValidation/Private/EditorValidatorSubsystem.cpp
@@ -843,6 +843,7 @@ EDataValidationResult UEditorValidatorSubsystem::ValidateChangelistsInternal(
 
 	if (AssetRegistry.IsLoadingAssets())
 	{
+		UE_CLOG(FApp::IsUnattended(), LogContentValidation, Fatal, TEXT("Unable to perform unattended content validation while asset registry scan is in progress. Callers just wait for asset registry scan to complete."));
 		FMessageDialog::Open(EAppMsgType::Ok, LOCTEXT("DataValidation.UnableToValidate_PendingAssetRegistry", "Unable to validate changelist while asset registry scan is in progress. Wait until asset discovery is complete."));
 		return EDataValidationResult::NotValidated;
 	}
