[Validation] Fix FScopedIgnoreLogMessages cleanup

diff --git a/Engine/Plugins/Editor/DataValidation/Source/DataValidation/Private/EditorValidatorHelpers.cpp b/Engine/Plugins/Editor/DataValidation/Source/DataValidation/Private/EditorValidatorHelpers.cpp
--- a/Engine/Plugins/Editor/DataValidation/Source/DataValidation/Private/EditorValidatorHelpers.cpp
+++ b/Engine/Plugins/Editor/DataValidation/Source/DataValidation/Private/EditorValidatorHelpers.cpp
@@ -244,7 +244,7 @@ TTuple<int32, int32> FScopedLogMessageGatherer::AddIgnoreCategories(TConstArrayV
 
 void FScopedLogMessageGatherer::RemoveIgnoreCategories(TTuple<int32, int32> Range)
 {
-    Impl->RemoveIgnorePatterns(Range);
+    Impl->RemoveIgnoreCategories(Range);
 }
 
 TTuple<int32, int32> FScopedLogMessageGatherer::AddIgnorePatterns(TConstArrayView<FWildcardString> NewPatterns)
@@ -283,7 +283,7 @@ FScopedIgnoreLogMessages::~FScopedIgnoreLogMessages()
     if(FScopedLogMessageGatherer* Gatherer = FScopedLogMessageGatherer::GetCurrentThreadGatherer())
     {
         Gatherer->RemoveIgnoreCategories(CategoriesRange);
-        Gatherer->RemoveIgnoreCategories(PatternsRange);
+        Gatherer->RemoveIgnorePatterns(PatternsRange);
     }
 }
 
