Public part for the improvements on texture validation from CL 29932763

- Added 3 new tags to UTexture2D to track properties necessary to check validity from FAssetData only
- Improved messaging (Thanks to [REDACTED_PERSON])
- Added an iconon textures' thumbnails to indicate invalidity
- Added invalidity tooltips in hovering window on textures when applicable
- Added contextual  menu item 'Conform Texture' in Texture sction when a texture is selected
- Added actual method to conform textures to validation criteria
- Added logic to automatically conform textures in pProjects created before 28.10 whenever necessary
- Added small texture criteria for UI and 2D textures which will hide any missed validation criteria
- Added texture budget criteria to monitor size in memory of non-streamable textures (only done when validating the project itself)
- Added ability to silently run validator subsystem

Note: preflight at CL 32229072 failed on 3 steps: EdTest, Verse and Content. EdTest and Verse also failed without the changes of this CL.The step  Content generated an invalid error on an uninitialzed property which was actually initialized.

diff --git a/Engine/Plugins/Editor/DataValidation/Source/DataValidation/Private/EditorValidatorSubsystem.cpp b/Engine/Plugins/Editor/DataValidation/Source/DataValidation/Private/EditorValidatorSubsystem.cpp
--- a/Engine/Plugins/Editor/DataValidation/Source/DataValidation/Private/EditorValidatorSubsystem.cpp
+++ b/Engine/Plugins/Editor/DataValidation/Source/DataValidation/Private/EditorValidatorSubsystem.cpp
@@ -388,7 +388,10 @@ EDataValidationResult UEditorValidatorSubsystem::ValidateAssetsInternal(
 	IAssetRegistry& AssetRegistry = IAssetRegistry::GetChecked();
 
 	FScopedSlowTask SlowTask(AssetDataList.Num(), LOCTEXT("DataValidation.ValidateAssetsTask", "Validating Assets"));
-	SlowTask.MakeDialog();
+	if (!InSettings.bSilent)
+	{
+		SlowTask.MakeDialog();
+	}
 	
 	UE_LOG(LogContentValidation, Display, TEXT("Starting to validate %d assets"), AssetDataList.Num());
 	UE_LOG(LogContentValidation, Log, TEXT("Enabled validators:"));
@@ -441,7 +444,7 @@ EDataValidationResult UEditorValidatorSubsystem::ValidateAssetsInternal(
 	}
 
 	// Dont let other async compilation warnings be attributed incorrectly to the package that is loading.
-	WaitForAssetCompilationIfNecessary(InSettings.ValidationUsecase);
+	WaitForAssetCompilationIfNecessary(InSettings.ValidationUsecase, !InSettings.bSilent);
 
 	OutResults.NumRequested = AssetDataList.Num();
 	
@@ -452,7 +455,10 @@ EDataValidationResult UEditorValidatorSubsystem::ValidateAssetsInternal(
 	{
 		ensure(Data.IsValid());
 
-		SlowTask.EnterProgressFrame(1.0f, FText::Format(LOCTEXT("DataValidation.ValidatingFilename", "Validating {0}"), FText::FromString(Data.GetFullName())));
+		if (!InSettings.bSilent)
+		{
+			SlowTask.EnterProgressFrame(1.0f, FText::Format(LOCTEXT("DataValidation.ValidatingFilename", "Validating {0}"), FText::FromString(Data.GetFullName())));
+		}
 		
 		if (OutResults.NumChecked >= InSettings.MaxAssetsToValidate)
 		{
@@ -837,7 +843,10 @@ EDataValidationResult UEditorValidatorSubsystem::ValidateChangelistsInternal(
 {
 	FScopedSlowTask SlowTask(Changelists.Num(), LOCTEXT("DataValidation.ValidatingChangelistTask", "Validating Changelists"));
 	SlowTask.Visibility = ESlowTaskVisibility::Invisible;
-	SlowTask.MakeDialog();
+	if (!Settings.bSilent)
+	{
+		SlowTask.MakeDialog();
+	}
 
 	IAssetRegistry& AssetRegistry = IAssetRegistry::GetChecked();
 
@@ -873,7 +882,10 @@ EDataValidationResult UEditorValidatorSubsystem::ValidateChangelistsInternal(
 	{
 		FText ValidationMessage = FText::Format(LOCTEXT("DataValidation.ValidatingChangelistMessage", "Validating changelist {0}"), Changelist->Description);
 		DataValidationLog.Info(ValidationMessage);
-		SlowTask.EnterProgressFrame(1.0f, ValidationMessage);
+		if (!Settings.bSilent)
+		{
+			SlowTask.EnterProgressFrame(1.0f, ValidationMessage);
+		}
 
 		FValidateAssetsDetails& Details = OutResults.AssetsDetails.FindOrAdd(Changelist->GetPathName());
 		{
@@ -1040,7 +1052,7 @@ TArray<FAssetData> UEditorValidatorSubsystem::GetAssetsResolvingRedirectors(FARF
 	return Found;
 }
 
-void UEditorValidatorSubsystem::WaitForAssetCompilationIfNecessary(EDataValidationUsecase InUsecase) const
+void UEditorValidatorSubsystem::WaitForAssetCompilationIfNecessary(EDataValidationUsecase InUsecase, bool bShowProgress) const
 {
 	if (InUsecase == EDataValidationUsecase::Save)
 	{
@@ -1049,8 +1061,11 @@ void UEditorValidatorSubsystem::WaitForAssetCompilationIfNecessary(EDataValidati
 
 	if (FAssetCompilingManager::Get().GetNumRemainingAssets())
 	{
-		FScopedSlowTask CompileAssetsSlowTask(0.f, LOCTEXT("DataValidation.CompilingAssetsBeforeCheckingContentTask", "Finishing asset compilations before checking content..."));
-		CompileAssetsSlowTask.MakeDialog();
+		if (bShowProgress)
+		{
+			FScopedSlowTask CompileAssetsSlowTask(0.f, LOCTEXT("DataValidation.CompilingAssetsBeforeCheckingContentTask", "Finishing asset compilations before checking content..."));
+			CompileAssetsSlowTask.MakeDialog();
+		}
 		FAssetCompilingManager::Get().FinishAllCompilation();
 	}
 }
diff --git a/Engine/Plugins/Editor/DataValidation/Source/DataValidation/Public/EditorValidatorSubsystem.h b/Engine/Plugins/Editor/DataValidation/Source/DataValidation/Public/EditorValidatorSubsystem.h
--- a/Engine/Plugins/Editor/DataValidation/Source/DataValidation/Public/EditorValidatorSubsystem.h
+++ b/Engine/Plugins/Editor/DataValidation/Source/DataValidation/Public/EditorValidatorSubsystem.h
@@ -160,9 +160,12 @@ struct DATAVALIDATION_API FValidateAssetsSettings
 	 * Defaults to warning, can be disabled by emptying the optional.
 	 */
 	TOptional<EMessageSeverity::Type> ShowMessageLogSeverity;
-	
+
 	/** Title of message log page to use for warnings/errors/etc */
 	FText MessageLogPageTitle;
+
+	/** Show progress window */
+	bool bSilent = false;
 };
 
 /**
@@ -338,7 +341,7 @@ public:
 protected:
 	void CleanupValidators();
 	
-	void WaitForAssetCompilationIfNecessary(EDataValidationUsecase InUsecase) const;
+	void WaitForAssetCompilationIfNecessary(EDataValidationUsecase InUsecase, bool bShowProgress = true) const;
 
 	/**
 	 * @return Returns true if the current Path should be skipped for validation. Returns false otherwise.
diff --git a/Engine/Source/Runtime/Engine/Private/Texture2D.cpp b/Engine/Source/Runtime/Engine/Private/Texture2D.cpp
--- a/Engine/Source/Runtime/Engine/Private/Texture2D.cpp
+++ b/Engine/Source/Runtime/Engine/Private/Texture2D.cpp
@@ -711,6 +711,13 @@ void UTexture2D::GetAssetRegistryTags(FAssetRegistryTagsContext Context) const
 	//	it does not tell you if the texture actually has non-opaque alpha
 	Context.AddTag( FAssetRegistryTag("HasAlphaChannel", HasAlphaChannel() ? TEXT("True") : TEXT("False"), FAssetRegistryTag::TT_Alphabetical) );
 	Context.AddTag( FAssetRegistryTag("Format", GPixelFormats[GetPixelFormat()].Name, FAssetRegistryTag::TT_Alphabetical) );
+#if WITH_EDITORONLY_DATA
+	const FString PowerOfTwoModeStr = StaticEnum<ETexturePowerOfTwoSetting::Type>()->GetNameStringByValue(static_cast<int64>(PowerOfTwoMode));
+	Context.AddTag(FAssetRegistryTag("PowerOfTwoMode", *PowerOfTwoModeStr, FAssetRegistryTag::TT_Alphabetical));
+	const FString MipGenSettingsStr = StaticEnum<TextureMipGenSettings>()->GetNameStringByValue(static_cast<int64>(MipGenSettings));
+	Context.AddTag(FAssetRegistryTag("MipGenSettings", *MipGenSettingsStr, FAssetRegistryTag::TT_Alphabetical));
+	Context.AddTag(FAssetRegistryTag("MaxTextureSize", FString::FromInt(MaxTextureSize), FAssetRegistryTag::TT_Numerical));
+#endif
 
 	Super::GetAssetRegistryTags(Context);
 }
