Fixed UEditorValidatorSubsystem::ValidateAssetsWithSettings taking its settings as an out-param rather than an in-param

diff --git a/Engine/Plugins/Editor/DataValidation/Source/DataValidation/Private/EditorValidatorSubsystem.cpp b/Engine/Plugins/Editor/DataValidation/Source/DataValidation/Private/EditorValidatorSubsystem.cpp
--- a/Engine/Plugins/Editor/DataValidation/Source/DataValidation/Private/EditorValidatorSubsystem.cpp
+++ b/Engine/Plugins/Editor/DataValidation/Source/DataValidation/Private/EditorValidatorSubsystem.cpp
@@ -365,13 +365,13 @@ EDataValidationResult UEditorValidatorSubsystem::ValidateObjectInternal(
 
 int32 UEditorValidatorSubsystem::ValidateAssetsWithSettings(
 	const TArray<FAssetData>& AssetDataList,
-	FValidateAssetsSettings& InSettings,
+	const FValidateAssetsSettings& InSettings,
 	FValidateAssetsResults& OutResults) const
 {
 	FMessageLog DataValidationLog(UE::DataValidation::MessageLogName);
 	ValidateAssetsInternal(DataValidationLog, TSet<FAssetData>{AssetDataList}, InSettings, OutResults);
 	
-	if (EMessageSeverity::Type* Severity = InSettings.ShowMessageLogSeverity.GetPtrOrNull())
+	if (const EMessageSeverity::Type* Severity = InSettings.ShowMessageLogSeverity.GetPtrOrNull())
 	{
 		DataValidationLog.Open(*Severity, false);
 	}
diff --git a/Engine/Plugins/Editor/DataValidation/Source/DataValidation/Public/EditorValidatorSubsystem.h b/Engine/Plugins/Editor/DataValidation/Source/DataValidation/Public/EditorValidatorSubsystem.h
--- a/Engine/Plugins/Editor/DataValidation/Source/DataValidation/Public/EditorValidatorSubsystem.h
+++ b/Engine/Plugins/Editor/DataValidation/Source/DataValidation/Public/EditorValidatorSubsystem.h
@@ -207,7 +207,7 @@ public:
 	UFUNCTION(BlueprintCallable, Category = "Asset Validation")
 	virtual int32 ValidateAssetsWithSettings(
 		const TArray<FAssetData>& AssetDataList,
-		FValidateAssetsSettings& InSettings,
+		const FValidateAssetsSettings& InSettings,
 		FValidateAssetsResults& OutResults) const;
 
 	/**
