#UE Check that modified packages pass the ShouldValidateAsset filter before processing them to avoid uncooked or non-runtime assets

diff --git a/Engine/Plugins/Editor/DataValidation/Source/DataValidation/Private/WorldPartitionChangelistValidator.cpp b/Engine/Plugins/Editor/DataValidation/Source/DataValidation/Private/WorldPartitionChangelistValidator.cpp
--- a/Engine/Plugins/Editor/DataValidation/Source/DataValidation/Private/WorldPartitionChangelistValidator.cpp
+++ b/Engine/Plugins/Editor/DataValidation/Source/DataValidation/Private/WorldPartitionChangelistValidator.cpp
@@ -6,6 +6,8 @@
 
 #include "AssetRegistry/AssetData.h"
 #include "AssetRegistry/AssetRegistryModule.h"
+#include "Editor.h"
+#include "EditorValidatorSubsystem.h"
 #include "ISourceControlModule.h"
 #include "Misc/PackageName.h"
 #include "SourceControlHelpers.h"
@@ -87,6 +89,10 @@ void UWorldPartitionChangelistValidator::ValidateActorsAndDataLayersFromChangeLi
 
 		return nullptr;
 	};
+
+	UEditorValidatorSubsystem* EditorValidationSubsystem = GEditor ? GEditor->GetEditorSubsystem<UEditorValidatorSubsystem>() : nullptr;
+	FValidateAssetsSettings Settings;
+	FDataValidationContext ValidationContext(false, Settings.ValidationUsecase, {});
 	
 	FAssetRegistryModule& AssetRegistryModule = FModuleManager::LoadModuleChecked<FAssetRegistryModule>(FName("AssetRegistry"));
 	IAssetRegistry& AssetRegistry = AssetRegistryModule.Get();
@@ -97,44 +103,47 @@ void UWorldPartitionChangelistValidator::ValidateActorsAndDataLayersFromChangeLi
 		
 		for (FAssetData& AssetData : PackageAssetsData)
 		{
-			if (UClass* ActorNativeClass = TryAssociateActorToMap(AssetData))
-			{
-				bSubmittingWorldDataLayers = ActorNativeClass->IsChildOf<AWorldDataLayers>();
-			}
-			else if (UClass* AssetClass = AssetData.GetClass())
+			if (!EditorValidationSubsystem || EditorValidationSubsystem->ShouldValidateAsset(AssetData, Settings, ValidationContext))
 			{
-				if (AssetClass->IsChildOf<UDataLayerAsset>())
+				if (UClass* ActorNativeClass = TryAssociateActorToMap(AssetData))
 				{
-					TArray<FName> ReferencerNames;
-					AssetRegistry.GetReferencers(AssetData.PackageName, ReferencerNames, UE::AssetRegistry::EDependencyCategory::All);
-		
-					FARFilter Filter;
-					Filter.bIncludeOnlyOnDiskAssets = true;
-					Filter.PackageNames = MoveTemp(ReferencerNames);
+					bSubmittingWorldDataLayers = ActorNativeClass->IsChildOf<AWorldDataLayers>();
+				}
+				else if (UClass* AssetClass = AssetData.GetClass())
+				{
+					if (AssetClass->IsChildOf<UDataLayerAsset>())
+					{
+						TArray<FName> ReferencerNames;
+						AssetRegistry.GetReferencers(AssetData.PackageName, ReferencerNames, UE::AssetRegistry::EDependencyCategory::All);
 
-					TArray<FAssetData> DataLayerReferencers;
-					AssetRegistry.GetAssets(Filter, DataLayerReferencers);
+						FARFilter Filter;
+						Filter.bIncludeOnlyOnDiskAssets = true;
+						Filter.PackageNames = MoveTemp(ReferencerNames);
 
-					for (const FAssetData& DataLayerReferencer : DataLayerReferencers)
-					{
-						UClass* ReferencerAssetClass = DataLayerReferencer.GetClass();
-						if (ReferencerAssetClass && ReferencerAssetClass->IsChildOf<AWorldDataLayers>())
+						TArray<FAssetData> DataLayerReferencers;
+						AssetRegistry.GetAssets(Filter, DataLayerReferencers);
+
+						for (const FAssetData& DataLayerReferencer : DataLayerReferencers)
 						{
-							TryAssociateActorToMap(DataLayerReferencer);
+							UClass* ReferencerAssetClass = DataLayerReferencer.GetClass();
+							if (ReferencerAssetClass && ReferencerAssetClass->IsChildOf<AWorldDataLayers>())
+							{
+								TryAssociateActorToMap(DataLayerReferencer);
+							}
 						}
+
+						RelevantDataLayerAssets.Add(AssetData.PackageName.ToString());
 					}
-		
-					RelevantDataLayerAssets.Add(AssetData.PackageName.ToString());
-				}
-				else if (AssetClass->IsChildOf<UDataLayerInstance>())
-				{
-					RelevantExternalPackageDataLayerInstances.Add(AssetData.PackageName.ToString());
-				}
-				else if (AssetClass->IsChildOf<UWorld>())
-				{
-					if (ULevel::GetIsLevelPartitionedFromPackage(PackageName))
+					else if (AssetClass->IsChildOf<UDataLayerInstance>())
+					{
+						RelevantExternalPackageDataLayerInstances.Add(AssetData.PackageName.ToString());
+					}
+					else if (AssetClass->IsChildOf<UWorld>())
 					{
-						MapToActorsFiles.FindOrAdd(AssetData.GetSoftObjectPath().GetAssetPath());
+						if (ULevel::GetIsLevelPartitionedFromPackage(PackageName))
+						{
+							MapToActorsFiles.FindOrAdd(AssetData.GetSoftObjectPath().GetAssetPath());
+						}
 					}
 				}
 			}
