#UE Do not attempt to scan packages from invalid paths in ValidateActorsAndDataLayersFromChangeList. Just log them and move on

diff --git a/Engine/Plugins/Editor/DataValidation/Source/DataValidation/Private/WorldPartitionChangelistValidator.cpp b/Engine/Plugins/Editor/DataValidation/Source/DataValidation/Private/WorldPartitionChangelistValidator.cpp
--- a/Engine/Plugins/Editor/DataValidation/Source/DataValidation/Private/WorldPartitionChangelistValidator.cpp
+++ b/Engine/Plugins/Editor/DataValidation/Source/DataValidation/Private/WorldPartitionChangelistValidator.cpp
@@ -74,9 +74,17 @@ void UWorldPartitionChangelistValidator::ValidateActorsAndDataLayersFromChangeLi
 
 			if (!ActorFiles)
 			{
-				if (ULevel::GetIsLevelPartitionedFromPackage(ActorPath.GetLongPackageFName()))
+				FName PackageName = ActorPath.GetLongPackageFName();
+				if (FPackageName::IsValidLongPackageName(PackageName.ToString()))
 				{
-					ActorFiles = &MapToActorsFiles.Add(MapAssetName);
+					if (ULevel::GetIsLevelPartitionedFromPackage(PackageName))
+					{
+						ActorFiles = &MapToActorsFiles.Add(MapAssetName);
+					}
+				}
+				else
+				{
+					UE_LOG(LogContentValidation, Log, TEXT("ValidateActorsAndDataLayersFromChangeList found a map with invalid path '%s'. It may be in an umounted content folder. Referencing actor package: %s"), *PackageName.ToString(), *AssetData.PackageName.ToString());
 				}
 			}
 
