Add a data validator that checks for corruption in package files.

- Currently the validator checks for the following:
-- Check that the package file ends with PACKAGE_FILE_TAG, this is the most common form of file corruption that we see, where the file is truncated or the end of the file is filled with zeros as the write never completed.
-- Serialize the package file summary from the start of the file which will make sure that it is in the format that we expect and that the versioning etc are correct.
-- Make sure that the package trailer can be loaded if the package file is supposed to have one.
-- Ensure that the locally stored payloads can be loaded from the package trailer.
-- Lastly double check that the payloads we load from the trailer have the hash that we expect, to make sure that they have not been altered since that point.
- The first two checks would be done just by trying to load an asset from the package file so are quite commonly caught without the need of an additional data validation, but even so the file could end up being corrupted at some point after the user last used it in the editor and the point of submission so it doesn't hurt to check. Do note that the checks that we do have are all based on corrupted package files that we have seen users successfully submit to a project causing build breakages that this validator would've prevented.
- We want the validator to be able to run on presubmit/commandlet level, in which case it will generally be run on a changelist in which case the validator will be given a single asset of type 'UDataValidationChangelist' which contains all of the modified packages in that changelist. We also want to allow users to invoke the validator on a package if they are suspicious of it, in which case the validator will receive the higher level asset as a UObject* from which we should be able to find the package file.
- Don't run the validator when saving packages as we don't really care if the package file is corrupted before overwriting it.
- Don't run the validator on in memory packages, we need a package file to operate on!

diff --git a/Engine/Plugins/Editor/DataValidation/Source/DataValidation/Private/PackageFileValidator.cpp b/Engine/Plugins/Editor/DataValidation/Source/DataValidation/Private/PackageFileValidator.cpp
new file mode 100644
--- /dev/null
+++ b/Engine/Plugins/Editor/DataValidation/Source/DataValidation/Private/PackageFileValidator.cpp
@@ -0,0 +1,193 @@
+// Copyright Epic Games, Inc. All Rights Reserved.
+
+#include "PackageFileValidator.h"
+
+#include "DataValidationChangelist.h"
+#include "HAL/IConsoleManager.h"
+#include "HAL/FileManager.h"
+#include "Misc/DataValidation.h"
+#include "UObject/PackageFileSummary.h"
+#include "UObject/PackageTrailer.h"
+
+#include UE_INLINE_GENERATED_CPP_BY_NAME(PackageFileValidator)
+
+#define LOCTEXT_NAMESPACE "PackageFileValidator"
+
+bool UPackageFileValidator::CanValidateAsset_Implementation(const FAssetData& AssetData, UObject* Asset, FDataValidationContext& Context) const
+{
+	// We don't want to validate the package on disk when saving, as we will be overwriting that file anyway
+	if (Context.GetValidationUsecase() == EDataValidationUsecase::Save)
+	{
+		return false;
+	}
+
+	if (Asset == nullptr)
+	{
+		return false;
+	}
+
+	if (UDataValidationChangelist::StaticClass() == Asset->GetClass())
+	{
+		return true;
+	}
+
+	if (UPackage* Package = Asset->GetPackage())
+	{
+		// Memory only packages don't have any files to validate
+		return !Package->HasAnyPackageFlags(PKG_InMemoryOnly);
+	}
+
+	return false;
+}
+
+EDataValidationResult UPackageFileValidator::ValidateLoadedAsset_Implementation(const FAssetData& AssetData, UObject* Asset, FDataValidationContext& Context)
+{
+	TRACE_CPUPROFILER_EVENT_SCOPE(UPackageFileValidator::ValidateLoadedAsset_Implementation);
+
+	TArray<FName> PackageNames = FindPackageNames(Asset);
+
+	int32 NumFailedPackages = 0;
+
+	for (FName PackageName : PackageNames)
+	{
+		FPackagePath PackagePath;
+		if (!FPackagePath::TryFromPackageName(PackageName, PackagePath))
+		{
+			NumFailedPackages++;
+			AssetFails(nullptr, FText::Format(LOCTEXT("NoPackageTag", "{0} Unable to open for reading"), FText::FromName(PackageName)));
+			
+			continue;
+		}
+
+		TUniquePtr<FArchive> PackageAr(IFileManager::Get().CreateFileReader(*PackagePath.GetLocalFullPath(), FILEREAD_Silent));
+		if (!PackageAr.IsValid())
+		{
+			NumFailedPackages++;
+			AssetFails(nullptr, FText::Format(LOCTEXT("NoPackageTag", "{0} Unable to open for reading"), FText::FromName(PackageName)));
+			
+			continue;
+		}
+
+		FPackageFileSummary Summary;
+		if (!ValidatePackageSummary(PackageName, *PackageAr, Context, Summary))
+		{
+			NumFailedPackages++;
+			continue;
+		}
+
+		if (Summary.PayloadTocOffset > 0 && !ValidatePackageTrailer(PackageName, *PackageAr, Context))
+		{
+			NumFailedPackages++;
+			continue;
+		}
+	}
+
+	return NumFailedPackages == 0 ? EDataValidationResult::Valid : EDataValidationResult::Invalid;
+}
+
+bool UPackageFileValidator::ValidatePackageSummary(FName PackageName, FArchive& Ar, FDataValidationContext& Context, FPackageFileSummary& OutSummary)
+{
+	TRACE_CPUPROFILER_EVENT_SCOPE(UPackageFileValidator::ValidatePackageSummary);
+
+	const int64 TagOffset = Ar.TotalSize() - sizeof(uint32);
+	Ar.Seek(TagOffset);
+
+	uint32 Tag = 0;
+	Ar << Tag;
+
+	if (Tag != PACKAGE_FILE_TAG || Ar.IsError())
+	{
+		AssetFails(nullptr, FText::Format(LOCTEXT("BadPkgTag", "{0} The end of package tag is not valid, the file is probably corrupt"), FText::FromName(PackageName)));
+		return false;
+	}
+
+	Ar.Seek(0);
+
+	Ar << OutSummary;
+
+	if (Ar.IsError() || OutSummary.Tag != PACKAGE_FILE_TAG)
+	{
+		AssetFails(nullptr, FText::Format(LOCTEXT("BadPkgSummary", "{0} Failed to read the package file summary, the file is probably corrupt"), FText::FromName(PackageName)));
+		return false;
+	}
+
+	if (OutSummary.IsFileVersionTooOld())
+	{
+		AssetFails(nullptr, FText::Format(LOCTEXT("PkgOutOfDate", "{0} is out of date and is not backwards compatible with the current process. Min Required Version: {1}  Package Version: {2}"),
+			FText::FromName(PackageName),
+			(int32)VER_UE4_OLDEST_LOADABLE_PACKAGE,
+			OutSummary.GetFileVersionUE().FileVersionUE4));
+		return false;
+	}
+
+	return true;
+}
+
+bool UPackageFileValidator::ValidatePackageTrailer(FName PackageName, FArchive& Ar, FDataValidationContext& Context)
+{
+	TRACE_CPUPROFILER_EVENT_SCOPE(UPackageFileValidator::ValidatePackageTrailer);
+
+	using namespace UE;
+
+	FPackageTrailer Trailer;
+	if (!FPackageTrailer::TryLoadFromArchive(Ar, Trailer))
+	{
+		AssetFails(nullptr, FText::Format(LOCTEXT("BadPkgTrailer", "{0} Failed to read the package trailer, the file is probably corrupt"), FText::FromName(PackageName)));
+		return false;
+	}
+
+	TArray<FIoHash> LocalPayloads = Trailer.GetPayloads(EPayloadStorageType::Local);
+
+	for (const FIoHash& Id : LocalPayloads)
+	{
+		FCompressedBuffer Payload = Trailer.LoadLocalPayload(Id, Ar);
+		if (Payload.IsNull())
+		{
+			AssetFails(nullptr, FText::Format(LOCTEXT("BadPayload", "{0} Failed to read the payload {1}, the file is probably corrupt"), FText::FromName(PackageName), FText::FromString(LexToString(Id))));
+			return false;
+		}
+
+		if (Id != Payload.GetRawHash())
+		{
+			AssetFails(nullptr, FText::Format(LOCTEXT("BadPayloadId", "{0} Failed to read the payload {1}, the file is probably corrupt"), FText::FromName(PackageName), FText::FromString(LexToString(Id))));
+			return false;
+		}
+
+		if (bValidatePayloadHashes)
+		{
+			TRACE_CPUPROFILER_EVENT_SCOPE(UPackageFileValidator::ValidatePackageTrailer::HashPayload);
+
+			FIoHash PayloadHash = FIoHash::HashBuffer(Payload.Decompress());
+			if (Id != PayloadHash)
+			{
+				AssetFails(nullptr, FText::Format(LOCTEXT("BadPayloadData", "{0} The payload data did not match it's stored hash {1} vs {2}, the file is probably corrupt"),
+					FText::FromName(PackageName),
+					FText::FromString(LexToString(Id)),
+					FText::FromString(LexToString(PayloadHash))));
+
+				return false;
+			}
+		}
+	}
+
+	return true;
+}
+
+TArray<FName> UPackageFileValidator::FindPackageNames(UObject* Asset) const
+{
+	if (UDataValidationChangelist* ChangeList = Cast<UDataValidationChangelist>(Asset))
+	{
+		return ChangeList->ModifiedPackageNames;
+	}
+	else if(UPackage* Package = Asset->GetPackage())
+	{
+		TArray<FName> PackageNames;
+		PackageNames.Add(Package->GetFName());
+
+		return PackageNames;
+	}
+
+	return TArray<FName>();
+}
+
+#undef LOCTEXT_NAMESPACE
diff --git a/Engine/Plugins/Editor/DataValidation/Source/DataValidation/Private/PackageFileValidator.h b/Engine/Plugins/Editor/DataValidation/Source/DataValidation/Private/PackageFileValidator.h
new file mode 100644
--- /dev/null
+++ b/Engine/Plugins/Editor/DataValidation/Source/DataValidation/Private/PackageFileValidator.h
@@ -0,0 +1,37 @@
+// Copyright Epic Games, Inc. All Rights Reserved.
+
+#pragma once
+
+#include "EditorValidatorBase.h"
+
+#include "PackageFileValidator.generated.h"
+
+struct FPackageFileSummary;
+
+/**
+ * This validator checks the format of the package on disk to make sure that is has not become corrupted since it was last saved.
+ * 
+ * To disable the validator entirely, set ini:Editor:[/Script/DataValidation.PackageFileValidator]:bIsConfigDisabled=true
+ * To disable validation of payload hashes (which is much slower than the rest of the validation) set 
+ * ini:Editor:[/Script/DataValidation.PackageFileValidator]:bValidatePayloadHashes=true
+ */
+UCLASS()
+class DATAVALIDATION_API UPackageFileValidator : public UEditorValidatorBase
+{
+	GENERATED_BODY()
+
+protected:
+
+	virtual bool CanValidateAsset_Implementation(const FAssetData& AssetData, UObject* Asset, FDataValidationContext& Context) const override;
+	virtual EDataValidationResult ValidateLoadedAsset_Implementation(const FAssetData& AssetData, UObject* Asset, FDataValidationContext& Context) override;
+
+private:
+
+	bool ValidatePackageSummary(FName PackageName, FArchive& Ar, FDataValidationContext& Context, FPackageFileSummary& OutSummary);
+	bool ValidatePackageTrailer(FName PackageName, FArchive& Ar, FDataValidationContext& Context);
+
+	TArray<FName> FindPackageNames(UObject* Asset) const;
+
+	UPROPERTY(Config)
+	bool bValidatePayloadHashes = true;
+};
