Fix data validation summary message to consider EDataValidationResult value as well.

diff --git a/Engine/Plugins/Editor/DataValidation/Source/DataValidation/Private/EditorValidatorSubsystem.cpp b/Engine/Plugins/Editor/DataValidation/Source/DataValidation/Private/EditorValidatorSubsystem.cpp
--- a/Engine/Plugins/Editor/DataValidation/Source/DataValidation/Private/EditorValidatorSubsystem.cpp
+++ b/Engine/Plugins/Editor/DataValidation/Source/DataValidation/Private/EditorValidatorSubsystem.cpp
@@ -654,9 +654,9 @@ EDataValidationResult UEditorValidatorSubsystem::ValidateAssetsInternal(
 	return Result;
 }
 
-void UEditorValidatorSubsystem::LogAssetValidationSummary(FMessageLog& DataValidationLog, const FValidateAssetsSettings& InSettings, const FValidateAssetsResults& Results) const
+void UEditorValidatorSubsystem::LogAssetValidationSummary(FMessageLog& DataValidationLog, const FValidateAssetsSettings& InSettings, EDataValidationResult Result, const FValidateAssetsResults& Results) const
 {
-	const bool bFailed = (Results.NumInvalid > 0);
+	const bool bFailed = (Results.NumInvalid > 0) || Result != EDataValidationResult::Valid;
 	const bool bAtLeastOneWarning = (Results.NumWarnings > 0);
 
 	if (bFailed || bAtLeastOneWarning || InSettings.bShowIfNoFailures)
@@ -924,7 +924,7 @@ EDataValidationResult UEditorValidatorSubsystem::ValidateChangelistsInternal(
 	// Validate assets from all changelists
 	EDataValidationResult AssetResult = ValidateAssetsInternal(DataValidationLog, MoveTemp(AssetsToValidate), Settings, OutResults);
 	Result = CombineDataValidationResults(Result, AssetResult);
-	LogAssetValidationSummary(DataValidationLog, Settings, OutResults);
+	LogAssetValidationSummary(DataValidationLog, Settings, Result, OutResults);
 
 	return Result;
 }
diff --git a/Engine/Plugins/Editor/DataValidation/Source/DataValidation/Public/EditorValidatorSubsystem.h b/Engine/Plugins/Editor/DataValidation/Source/DataValidation/Public/EditorValidatorSubsystem.h
--- a/Engine/Plugins/Editor/DataValidation/Source/DataValidation/Public/EditorValidatorSubsystem.h
+++ b/Engine/Plugins/Editor/DataValidation/Source/DataValidation/Public/EditorValidatorSubsystem.h
@@ -383,7 +383,7 @@ protected:
 		const FValidateAssetsSettings& 				InSettings,
 		FValidateAssetsResults& 					OutResults) const;
 	
-	void LogAssetValidationSummary(FMessageLog& DataValidationLog, const FValidateAssetsSettings& InSettings, const FValidateAssetsResults& Results) const;
+	void LogAssetValidationSummary(FMessageLog& DataValidationLog, const FValidateAssetsSettings& InSettings, EDataValidationResult Result, const FValidateAssetsResults& Results) const;
 
 	EDataValidationResult ValidateObjectInternal(
 		const FAssetData& InAssetData,
