Added a setting to control whether warnings logged during validation are promoted to errors

FValidateAssetsSettings::bCaptureWarningsDuringValidationAsErrors controls this behavior, and defaults to true to preserve the existing behavior.

Note: FDataValidationModule::OnValidateSourcePackageDuringCook now keeps them as warnings, as this is only used by UE::Cook::FDelegates::ValidateSourcePackage and cannot accept a FValidateAssetsSettings argument.

diff --git a/Engine/Plugins/Editor/DataValidation/Source/DataValidation/Private/DataValidationModule.cpp b/Engine/Plugins/Editor/DataValidation/Source/DataValidation/Private/DataValidationModule.cpp
--- a/Engine/Plugins/Editor/DataValidation/Source/DataValidation/Private/DataValidationModule.cpp
+++ b/Engine/Plugins/Editor/DataValidation/Source/DataValidation/Private/DataValidationModule.cpp
@@ -426,10 +426,9 @@ EDataValidationResult FDataValidationModule::OnValidateSourcePackageDuringCook(U
 					{
 						TStringBuilder<2048> Buffer;
 						Buffer.Join(LogWarnings, LINE_TERMINATOR);
-						ValidationContext.AddMessage(EMessageSeverity::Error)
+						ValidationContext.AddMessage(EMessageSeverity::Warning)
 							->AddToken(FAssetDataToken::Create(AssetData))
 							->AddText(LOCTEXT("DataValidation.DuringValidationWarnings", "Warnings logged while validating asset {0}"), FText::FromStringView(Buffer.ToView()));
-						ValidationResult = EDataValidationResult::Invalid;
 					}
 					if (LogErrors.Num() > 0)
 					{
diff --git a/Engine/Plugins/Editor/DataValidation/Source/DataValidation/Private/EditorValidatorSubsystem.cpp b/Engine/Plugins/Editor/DataValidation/Source/DataValidation/Private/EditorValidatorSubsystem.cpp
--- a/Engine/Plugins/Editor/DataValidation/Source/DataValidation/Private/EditorValidatorSubsystem.cpp
+++ b/Engine/Plugins/Editor/DataValidation/Source/DataValidation/Private/EditorValidatorSubsystem.cpp
@@ -546,10 +546,13 @@ EDataValidationResult UEditorValidatorSubsystem::ValidateAssetsInternal(
 			{
 				TStringBuilder<2048> Buffer;
 				Buffer.Join(Warnings, LINE_TERMINATOR);
-				ValidationContext.AddMessage(EMessageSeverity::Error)
+				ValidationContext.AddMessage(InSettings.bCaptureWarningsDuringValidationAsErrors ? EMessageSeverity::Error : EMessageSeverity::Warning)
 					->AddToken(FAssetDataToken::Create(Data))
 					->AddText(LOCTEXT("DataValidation.DuringValidationWarnings", "Warnings logged while validating asset {0}"), FText::FromStringView(Buffer.ToView()));
-				AssetResult = EDataValidationResult::Invalid;
+				if (InSettings.bCaptureWarningsDuringValidationAsErrors)
+				{
+					AssetResult = EDataValidationResult::Invalid;
+				}
 			}
 			if(Errors.Num())
 			{
diff --git a/Engine/Plugins/Editor/DataValidation/Source/DataValidation/Public/EditorValidatorSubsystem.h b/Engine/Plugins/Editor/DataValidation/Source/DataValidation/Public/EditorValidatorSubsystem.h
--- a/Engine/Plugins/Editor/DataValidation/Source/DataValidation/Public/EditorValidatorSubsystem.h
+++ b/Engine/Plugins/Editor/DataValidation/Source/DataValidation/Public/EditorValidatorSubsystem.h
@@ -147,6 +147,10 @@ struct DATAVALIDATION_API FValidateAssetsSettings
 	UPROPERTY(EditAnywhere, BlueprintReadWrite, Category = "Validation")
 	bool bCaptureLogsDuringValidation = true;
 	
+	/** If true, captured log warnings during validation are added to the validation results as errors (requires bCaptureLogsDuringValidation) */
+	UPROPERTY(EditAnywhere, BlueprintReadWrite, Category = "Validation")
+	bool bCaptureWarningsDuringValidationAsErrors = true;
+
 	/** Maximum number of assets to attempt to validate */
 	UPROPERTY(EditAnywhere, BlueprintReadWrite, Category = "Validation")
 	int32 MaxAssetsToValidate = MAX_int32;
