Fixed uninitialized variable

FScopedLogMessageGatherer::Impl would be left uninitialized when constructing it with bEnabled set to false, which could lead to a crash in the destructor

#rnx

diff --git a/Engine/Plugins/Editor/DataValidation/Source/DataValidation/Public/EditorValidatorHelpers.h b/Engine/Plugins/Editor/DataValidation/Source/DataValidation/Public/EditorValidatorHelpers.h
--- a/Engine/Plugins/Editor/DataValidation/Source/DataValidation/Public/EditorValidatorHelpers.h
+++ b/Engine/Plugins/Editor/DataValidation/Source/DataValidation/Public/EditorValidatorHelpers.h
@@ -83,7 +83,7 @@ namespace UE::DataValidation
         
     protected:
         FScopedLogMessageGatherer* Previous = nullptr;
-        FLogMessageGathererImpl* Impl;
+        FLogMessageGathererImpl* Impl = nullptr;
     };
     
     // Utility to tell the closest enclosing FScopedLogMessageGatherer to ignore certain messages
