Editor | Confirm changelist submit window fixes

- Fix UDataValidationChangelist warnings not propagating to the window properly and only being emitted into the Message Log
- Fix asset validation warnings not appearing in the Confirm changelist submit window if the validator says the asset is valid
- Fix tokenized messages appearing blank

diff --git a/Engine/Plugins/Editor/DataValidation/Source/DataValidation/Private/EditorValidatorSubsystem.cpp b/Engine/Plugins/Editor/DataValidation/Source/DataValidation/Private/EditorValidatorSubsystem.cpp
--- a/Engine/Plugins/Editor/DataValidation/Source/DataValidation/Private/EditorValidatorSubsystem.cpp
+++ b/Engine/Plugins/Editor/DataValidation/Source/DataValidation/Private/EditorValidatorSubsystem.cpp
@@ -893,6 +893,7 @@ EDataValidationResult UEditorValidatorSubsystem::ValidateChangelistsInternal(
 			FDataValidationContext ValidationContext(false, Settings.ValidationUsecase, {}); // No associated objects for changelist
 			Details.Result = IsObjectValidWithContext(Changelist, ValidationContext);
 			UE::DataValidation::AddAssetValidationMessages(DataValidationLog, ValidationContext);
+			ValidationContext.SplitIssues(Details.ValidationWarnings, Details.ValidationErrors);
 		}
 		Result = CombineDataValidationResults(Result, Details.Result);
 		DataValidationLog.Flush();
diff --git a/Engine/Source/Editor/SourceControlWindows/Private/SSourceControlSubmit.cpp b/Engine/Source/Editor/SourceControlWindows/Private/SSourceControlSubmit.cpp
--- a/Engine/Source/Editor/SourceControlWindows/Private/SSourceControlSubmit.cpp
+++ b/Engine/Source/Editor/SourceControlWindows/Private/SSourceControlSubmit.cpp
@@ -282,7 +282,7 @@ void SSourceControlSubmitWidget::Construct(const FArguments& InArgs)
 		const FName ChangelistWarningsIconName = TEXT("Icons.WarningWithColor.Large");
 		const FName ChangelistErrorsIconName = TEXT("Icons.ErrorWithColor.Large");
 
-		if (bAllowSubmit)
+		if (ChangelistResultWarningsText.IsEmpty() && ChangelistResultErrorsText.IsEmpty())
 		{
 			Contents->AddSlot()
 			.AutoHeight()
diff --git a/Engine/Source/Runtime/CoreUObject/Private/Misc/DataValidation.cpp b/Engine/Source/Runtime/CoreUObject/Private/Misc/DataValidation.cpp
--- a/Engine/Source/Runtime/CoreUObject/Private/Misc/DataValidation.cpp
+++ b/Engine/Source/Runtime/CoreUObject/Private/Misc/DataValidation.cpp
@@ -86,11 +86,11 @@ void FDataValidationContext::SplitIssues(TArray<FText>& Warnings, TArray<FText>&
 
 		if (Issue.Severity == EMessageSeverity::Error)
 		{
-			Errors.Add(Issue.Message);
+			Errors.Add(MoveTemp(Message));
 		}
 		else if (Issue.Severity == EMessageSeverity::Warning)
 		{	
-			Warnings.Add(Issue.Message);
+			Warnings.Add(MoveTemp(Message));
 		}
 	}
 }
