Added verbose log when running a validator for a given asset

#rnx

diff --git a/Engine/Plugins/Editor/DataValidation/Source/DataValidation/Private/EditorValidatorSubsystem.cpp b/Engine/Plugins/Editor/DataValidation/Source/DataValidation/Private/EditorValidatorSubsystem.cpp
--- a/Engine/Plugins/Editor/DataValidation/Source/DataValidation/Private/EditorValidatorSubsystem.cpp
+++ b/Engine/Plugins/Editor/DataValidation/Source/DataValidation/Private/EditorValidatorSubsystem.cpp
@@ -353,6 +353,7 @@ EDataValidationResult UEditorValidatorSubsystem::ValidateObjectInternal(
 		
 		ForEachEnabledValidator([InObject, &InAssetData, &InContext, &Result](UEditorValidatorBase* Validator)
 		{
+			UE_LOG(LogContentValidation, Verbose, TEXT("Validating '%s' with '%s'..."), *InObject->GetPathName(), *Validator->GetClass()->GetName());
 			GInitRunaway(); // Reset runaway counter, as ValidateLoadedAsset may be implemented in a BP and could overflow the runaway count due to being called in a loop
 			EDataValidationResult NewResult = Validator->ValidateLoadedAsset(InAssetData, InObject, InContext);
 			Result = CombineDataValidationResults(Result, NewResult);
