Fix for deadlock in FLogMessageGathererImpl
* Moved GLog->RemoveOutputDevice(this) outside of the lock scope to ensure that we don't create a deadlock if Serialize is called while the message gatherer is being destroyed

#rnx

diff --git a/Engine/Plugins/Editor/DataValidation/Source/DataValidation/Private/EditorValidatorHelpers.cpp b/Engine/Plugins/Editor/DataValidation/Source/DataValidation/Private/EditorValidatorHelpers.cpp
--- a/Engine/Plugins/Editor/DataValidation/Source/DataValidation/Private/EditorValidatorHelpers.cpp
+++ b/Engine/Plugins/Editor/DataValidation/Source/DataValidation/Private/EditorValidatorHelpers.cpp
@@ -102,10 +102,8 @@ struct FLogMessageGathererImpl : public FOutputDevice
     
     ~FLogMessageGathererImpl()
     {
-        // Lock the impl, remove it as an output device, and delete it 
-        // There should be no problems with anyone waiting for the lock as access is controlled by the FScopedLogMessageGatherer
-        // and we are deleted when that goes out of scope
-        CS.Lock();
+        // Remove this as an output device, internally this handles synchronization with Serialize for us so we don't need to protect against Serialize and ~FLogMessageGathererImpl running concurrently 
+        // we deliberately don't lock the Critical Section (CS), as this would invert the locking order between FLogMessageGathererImpl & FOutputDeviceRedirector leading to potential a deadlock 
         GLog->RemoveOutputDevice(this);
     }
    
@@ -287,4 +285,4 @@ FScopedIgnoreLogMessages::~FScopedIgnoreLogMessages()
     }
 }
 
-} // namespace UE::DataValidation
\ No newline at end of file
+} // namespace UE::DataValidation
