Adding UMaterial and UMaterialInstance asset validators

diff --git a/Engine/Config/Android/DataDrivenPlatformInfo.ini b/Engine/Config/Android/DataDrivenPlatformInfo.ini
--- a/Engine/Config/Android/DataDrivenPlatformInfo.ini
+++ b/Engine/Config/Android/DataDrivenPlatformInfo.ini
@@ -133,4 +133,5 @@ bIsHlslcc=false
 bSupportsWaveOperations=RuntimeDependent
 MinimumWaveSize=4
 MaximumWaveSize=128
-bSupportsWave64=true
\ No newline at end of file
+bSupportsWave64=true
+FriendlyName=LOCTEXT("FriendlyShaderPlatformName_AndroidVulkanSM5", "Android Vulkan SM5 Mobile")
\ No newline at end of file
diff --git a/Engine/Config/Windows/DataDrivenPlatformInfo.ini b/Engine/Config/Windows/DataDrivenPlatformInfo.ini
--- a/Engine/Config/Windows/DataDrivenPlatformInfo.ini
+++ b/Engine/Config/Windows/DataDrivenPlatformInfo.ini
@@ -173,6 +173,7 @@ PreviewDeviceProfileName=Windows_Preview_ES31
 [PreviewShaderPlatform SP_PCD3D_ES3_1_SDF]
 ParentShaderPlatform=PCD3D_ES3_1
 bSupportsDistanceFields=true
+bCanUsePreviewPlatformForMaterialValidation=true
 FriendlyName=LOCTEXT("FriendlyShaderPlatformName_D3D_Mobile_SDF", "PC D3D Mobile SDF")
 
 [PreviewPlatform PCD3D_ES3_1_SDF]
diff --git a/Engine/Plugins/Editor/DataValidation/Source/DataValidation/DataValidation.Build.cs b/Engine/Plugins/Editor/DataValidation/Source/DataValidation/DataValidation.Build.cs
--- a/Engine/Plugins/Editor/DataValidation/Source/DataValidation/DataValidation.Build.cs
+++ b/Engine/Plugins/Editor/DataValidation/Source/DataValidation/DataValidation.Build.cs
@@ -46,6 +46,7 @@ public class DataValidation : ModuleRules
 				"TypedElementRuntime",
 				"SourceControl",
 				"UncontrolledChangelists",
+				"RHI"
 			}
 		);
 	}
diff --git a/Engine/Plugins/Editor/DataValidation/Source/DataValidation/Private/DataValidationModule.cpp b/Engine/Plugins/Editor/DataValidation/Source/DataValidation/Private/DataValidationModule.cpp
--- a/Engine/Plugins/Editor/DataValidation/Source/DataValidation/Private/DataValidationModule.cpp
+++ b/Engine/Plugins/Editor/DataValidation/Source/DataValidation/Private/DataValidationModule.cpp
@@ -13,6 +13,7 @@
 #include "Misc/MessageDialog.h"
 #include "DataValidationCommandlet.h"
 #include "EditorValidatorBase.h"
+#include "EditorValidator_Material.h"
 #include "Elements/Framework/TypedElementSelectionSet.h"
 #include "Logging/MessageLog.h"
 #include "UObject/ICookInfo.h"
@@ -76,6 +77,8 @@ void FDataValidationModule::StartupModule()
 			LOCTEXT("DataValidationDescription", "Settings related to validating assets in the editor."),
 			GetMutableDefault<UDataValidationSettings>()
 		);
+
+		FMaterialEditorValidationShaderPlatform::RegisterCustomPropertyTypeLayout();
 	}
 }
 
@@ -89,6 +92,8 @@ void FDataValidationModule::ShutdownModule()
 		UToolMenus::UnregisterOwner(this);
 
 		UPackage::PackageSavedWithContextEvent.RemoveAll(this);
+
+		FMaterialEditorValidationShaderPlatform::UnregisterCustomPropertyTypeLayout();
 	}
 }
 
diff --git a/Engine/Plugins/Editor/DataValidation/Source/DataValidation/Private/EditorValidator_Material.cpp b/Engine/Plugins/Editor/DataValidation/Source/DataValidation/Private/EditorValidator_Material.cpp
new file mode 100644
--- /dev/null
+++ b/Engine/Plugins/Editor/DataValidation/Source/DataValidation/Private/EditorValidator_Material.cpp
@@ -0,0 +1,276 @@
+﻿// Copyright Epic Games, Inc. All Rights Reserved.
+
+#include "EditorValidator_Material.h"
+#include "AssetCompilingManager.h"
+#include "DataDrivenShaderPlatformInfo.h"
+#include "DetailWidgetRow.h"
+#include "EditorValidatorSubsystem.h"
+#include "Materials/Material.h"
+#include "Materials/MaterialInstance.h"
+#include "Misc/DataValidation.h"
+#include "Widgets/Input/SComboBox.h"
+
+UEditorValidator_Material::UEditorValidator_Material()
+	: Super()
+{
+	for (const FMaterialEditorValidationPlatform& Config: GetDefault<UDataValidationSettings>()->MaterialValidationPlatforms)
+	{
+		FShaderValidationPlatform Platform = {};
+
+		bool bValidShaderPlatform = false;
+		if (Config.ShaderPlatform.Name == FMaterialEditorValidationShaderPlatform::MaxRHIShaderPlatformName)
+		{
+			Platform.ShaderPlatform = GMaxRHIShaderPlatform;
+			bValidShaderPlatform = true;
+		}
+		else
+		{
+			for (int32 ShaderPlatformIndex = 0; ShaderPlatformIndex < SP_NumPlatforms; ++ShaderPlatformIndex)
+			{
+				const EShaderPlatform ShaderPlatform = static_cast<EShaderPlatform>(ShaderPlatformIndex);
+
+				if (FDataDrivenShaderPlatformInfo::IsValid(ShaderPlatform)
+					&& FDataDrivenShaderPlatformInfo::CanUseForMaterialValidation(ShaderPlatform)
+					&& FDataDrivenShaderPlatformInfo::GetName(ShaderPlatform) == Config.ShaderPlatform.Name)
+				{
+					Platform.ShaderPlatform = ShaderPlatform;
+					bValidShaderPlatform = true;
+					break;
+				}
+			}
+		}
+
+		if (!bValidShaderPlatform)
+		{
+			UE_LOG(LogContentValidation, Warning, TEXT("Material asset validation shader platform '%s' is not available, skipping."), *Config.ShaderPlatform.Name.ToString());
+			continue;
+		}
+
+		switch (Config.FeatureLevel)
+		{
+		case EMaterialEditorValidationFeatureLevel::CurrentMaxFeatureLevel: Platform.FeatureLevel = GMaxRHIFeatureLevel; break; 
+		case EMaterialEditorValidationFeatureLevel::ES3_1: Platform.FeatureLevel = ERHIFeatureLevel::ES3_1; break;
+		case EMaterialEditorValidationFeatureLevel::SM5: Platform.FeatureLevel = ERHIFeatureLevel::SM5; break;
+		case EMaterialEditorValidationFeatureLevel::SM6: Platform.FeatureLevel = ERHIFeatureLevel::SM6; break;
+		}
+
+		switch (Config.MaterialQualityLevel)
+		{
+		case EMaterialEditorValidationQualityLevel::Low: Platform.MaterialQualityLevel = EMaterialQualityLevel::Low; break;
+		case EMaterialEditorValidationQualityLevel::Medium: Platform.MaterialQualityLevel = EMaterialQualityLevel::Medium; break;
+		case EMaterialEditorValidationQualityLevel::High: Platform.MaterialQualityLevel = EMaterialQualityLevel::High; break;
+		case EMaterialEditorValidationQualityLevel::Epic: Platform.MaterialQualityLevel = EMaterialQualityLevel::Epic; break;
+		}
+
+		ValidationPlatforms.Add(Platform);
+	}
+}
+
+bool UEditorValidator_Material::CanValidateAsset_Implementation(const FAssetData& AssetData, UObject* InAsset, FDataValidationContext& InContext) const
+{
+	const bool bAnyValidationPlatforms = !ValidationPlatforms.IsEmpty();
+	const bool bShouldAllowFullValidation = InContext.GetValidationUsecase() != EDataValidationUsecase::Save;
+	return bAnyValidationPlatforms && bShouldAllowFullValidation && (Cast<UMaterial>(InAsset) || Cast<UMaterialInstance>(InAsset));
+}
+
+EDataValidationResult UEditorValidator_Material::ValidateLoadedAsset_Implementation(const FAssetData& AssetData, UObject* InAsset, FDataValidationContext& InContext)
+{
+	UMaterialInstance* MaterialInstance = Cast<UMaterialInstance>(InAsset);
+
+	UMaterial* Material = MaterialInstance ? MaterialInstance->GetMaterial() : Cast<UMaterial>(InAsset);
+	check(Material);
+
+	TArray<FMaterialResource*> Resources;
+	bool bAnyResources = false;
+
+	for (const FShaderValidationPlatform& ValidationPlatform: ValidationPlatforms)
+	{
+		FMaterialResource* CurrentResource = FindOrCreateMaterialResource(Resources, Material, MaterialInstance, ValidationPlatform.FeatureLevel, ValidationPlatform.MaterialQualityLevel);
+		check(CurrentResource);
+
+		if (CurrentResource)
+		{
+			CurrentResource->CacheShaders(ValidationPlatform.ShaderPlatform);
+			bAnyResources = true;
+		}
+	}
+
+	if (!bAnyResources)
+	{
+		FMaterial::DeferredDeleteArray(Resources);
+		return EDataValidationResult::NotValidated;
+	}
+
+	FAssetCompilingManager::Get().FinishAllCompilation();
+
+	bool bCompileErrors = false;
+	for (FMaterialResource* Resource: Resources)
+	{
+		if (!Resource->IsCompilationFinished())
+		{
+			UE_LOG(LogContentValidation, Warning, TEXT("Shader compilation was expected to be finished, but was not finished."));
+		}
+
+		bCompileErrors |= Resource->GetCompileErrors().Num() > 0;
+	}
+
+	FMaterial::DeferredDeleteArray(Resources);
+
+	return !bCompileErrors ? EDataValidationResult::Valid : EDataValidationResult::Invalid;
+}
+
+class FMaterialEditorValidationPlatformCustomization : public IPropertyTypeCustomization
+{
+public:
+	FMaterialEditorValidationPlatformCustomization()
+		: MaxRHIShaderPlatform(MakeShared<EShaderPlatform>(SP_NumPlatforms))
+	{
+		for (int32 ShaderPlatformIndex = 0; ShaderPlatformIndex < SP_NumPlatforms; ++ShaderPlatformIndex)
+		{
+			const EShaderPlatform ShaderPlatform = static_cast<EShaderPlatform>(ShaderPlatformIndex);
+
+			if (FDataDrivenShaderPlatformInfo::IsValid(ShaderPlatform) && FDataDrivenShaderPlatformInfo::CanUseForMaterialValidation(ShaderPlatform))
+			{
+				ValidationShaderPlatforms.Add(MakeShared<EShaderPlatform>(ShaderPlatform));
+			}
+		}
+
+		ValidationShaderPlatforms.StableSort([this](const TSharedPtr<EShaderPlatform>& A, const TSharedPtr<EShaderPlatform>& B)
+		{
+			return GetShaderPlatformFriendlyName(A).CompareTo(GetShaderPlatformFriendlyName(B)) < 0;
+		});
+
+		ValidationShaderPlatforms.Insert(MaxRHIShaderPlatform, 0);
+	}
+
+	virtual void CustomizeHeader(TSharedRef<IPropertyHandle> StructPropertyHandle, FDetailWidgetRow& HeaderRow, IPropertyTypeCustomizationUtils& StructCustomizationUtils) override
+	{
+		TSharedPtr<IPropertyHandle> PropertyHandle = StructPropertyHandle->GetChildHandle(GET_MEMBER_NAME_CHECKED(FMaterialEditorValidationShaderPlatform, Name));
+		ensure(PropertyHandle.IsValid());
+		if (!PropertyHandle.IsValid())
+		{
+			return;
+		}
+
+		HeaderRow
+		.NameContent()
+		[
+			StructPropertyHandle->CreatePropertyNameWidget()
+		]
+		.ValueContent()
+		[
+			SNew(SComboBox<TSharedPtr<EShaderPlatform>>)
+			.OptionsSource(&ValidationShaderPlatforms)
+			.InitiallySelectedItem(GetCurrentShaderPlatform(PropertyHandle))
+			.OnSelectionChanged_Lambda([this, PropertyHandle](TSharedPtr<EShaderPlatform> ShaderPlatformOpt, ESelectInfo::Type SelectionInfo)
+			{
+				if (ShaderPlatformOpt.IsValid() && PropertyHandle.IsValid())
+				{
+					PropertyHandle->NotifyPreChange();
+					if (ShaderPlatformOpt == MaxRHIShaderPlatform)
+					{
+						PropertyHandle->SetValue(FMaterialEditorValidationShaderPlatform::MaxRHIShaderPlatformName);
+					}
+					else
+					{
+						PropertyHandle->SetValue(FDataDrivenShaderPlatformInfo::GetName(*ShaderPlatformOpt.Get()));
+					}
+					PropertyHandle->NotifyPostChange(EPropertyChangeType::ValueSet);
+				}
+			})
+			.OnGenerateWidget_Lambda([this](const TSharedPtr<EShaderPlatform>& ShaderPlatformOpt)
+			{
+				return SNew(STextBlock).Text(GetShaderPlatformFriendlyName(ShaderPlatformOpt));
+			})
+			.Content()
+			[
+				SNew(STextBlock)
+				.Font(StructCustomizationUtils.GetRegularFont())
+				.Text_Lambda([this, PropertyHandle]()
+				{
+					return GetShaderPlatformFriendlyName(GetCurrentShaderPlatform(PropertyHandle)); 
+				})
+			]
+		];
+	}
+
+	virtual void CustomizeChildren(TSharedRef<IPropertyHandle> StructPropertyHandle, IDetailChildrenBuilder& StructBuilder, IPropertyTypeCustomizationUtils& StructCustomizationUtils) override
+	{
+	}
+
+private:
+	TArray<TSharedPtr<EShaderPlatform>> ValidationShaderPlatforms;
+	const TSharedPtr<EShaderPlatform> MaxRHIShaderPlatform;
+
+	FName GetShaderPlatformName(const TSharedPtr<EShaderPlatform>& ShaderPlatformOpt)
+	{
+		if (ShaderPlatformOpt == MaxRHIShaderPlatform)
+		{
+			return FMaterialEditorValidationShaderPlatform::MaxRHIShaderPlatformName;
+		}
+		else if (ShaderPlatformOpt.IsValid())
+		{
+			return FDataDrivenShaderPlatformInfo::GetName(*ShaderPlatformOpt.Get());
+		}
+		else
+		{
+			return FName("Invalid");
+		}
+	}
+
+	FText GetShaderPlatformFriendlyName(const TSharedPtr<EShaderPlatform>& ShaderPlatformOpt)
+	{
+		if (ShaderPlatformOpt == MaxRHIShaderPlatform)
+		{
+			return NSLOCTEXT("AssetValidation", "ShaderPlatform_MaxRHIShaderPlatform", "Current RHI Max Shader Platform");
+		}
+		else if (ShaderPlatformOpt.IsValid())
+		{
+			const FText Name = FDataDrivenShaderPlatformInfo::GetFriendlyName(*ShaderPlatformOpt.Get());
+			if (!Name.IsEmpty())
+			{
+				return Name;
+			}
+
+			return FText::FromName(FDataDrivenShaderPlatformInfo::GetName(*ShaderPlatformOpt.Get()));
+		}
+		else
+		{
+			return NSLOCTEXT("AssetValidation", "ShaderPlatform_Invalid", "Invalid");
+		}
+	}
+
+	TSharedPtr<EShaderPlatform> GetCurrentShaderPlatform(const TSharedPtr<IPropertyHandle>& PropertyHandle)
+	{
+		FName CurrentShaderPlatformName;
+		if (PropertyHandle->GetValue(CurrentShaderPlatformName))
+		{
+			for (const TSharedPtr<EShaderPlatform>& ValidationShaderPlatform : ValidationShaderPlatforms)
+			{
+				if (GetShaderPlatformName(ValidationShaderPlatform) == CurrentShaderPlatformName)
+				{
+					return ValidationShaderPlatform;
+				}
+			}
+		}
+
+		return TSharedPtr<EShaderPlatform>();
+	}
+};
+
+FName FMaterialEditorValidationShaderPlatform::MaxRHIShaderPlatformName = MaxRHIShaderPlatformNameView.GetData();
+
+void FMaterialEditorValidationShaderPlatform::RegisterCustomPropertyTypeLayout()
+{
+	FPropertyEditorModule& PropertyModule = FModuleManager::LoadModuleChecked<FPropertyEditorModule>("PropertyEditor");
+	PropertyModule.RegisterCustomPropertyTypeLayout(
+		StaticStruct()->GetFName(),
+		FOnGetPropertyTypeCustomizationInstance::CreateLambda([](){return MakeShared<FMaterialEditorValidationPlatformCustomization>();})
+	);
+}
+
+void FMaterialEditorValidationShaderPlatform::UnregisterCustomPropertyTypeLayout()
+{
+	FPropertyEditorModule& PropertyModule = FModuleManager::LoadModuleChecked<FPropertyEditorModule>("PropertyEditor");
+	PropertyModule.UnregisterCustomPropertyTypeLayout(StaticStruct()->GetFName());
+}
diff --git a/Engine/Plugins/Editor/DataValidation/Source/DataValidation/Public/DataValidationSettings.h b/Engine/Plugins/Editor/DataValidation/Source/DataValidation/Public/DataValidationSettings.h
--- a/Engine/Plugins/Editor/DataValidation/Source/DataValidation/Public/DataValidationSettings.h
+++ b/Engine/Plugins/Editor/DataValidation/Source/DataValidation/Public/DataValidationSettings.h
@@ -3,6 +3,7 @@
 #pragma once
 
 #include "Engine/DeveloperSettings.h"
+#include "EditorValidator_Material.h"
 
 #include "DataValidationSettings.generated.h"
 
@@ -19,10 +20,13 @@ public:
 	UDataValidationSettings();
 
 	/** Whether or not to validate assets on save */
-	UPROPERTY(EditAnywhere, config, Category="Data Validation")
+	UPROPERTY(EditAnywhere, Config, Category="Data Validation")
 	uint32 bValidateOnSave : 1 = true;
 	
 	/** Whether or not to load & validate assets in changelists by default */
-	UPROPERTY(EditAnywhere, config, Category="Data Validation")
+	UPROPERTY(EditAnywhere, Config, Category="Data Validation")
 	uint32 bLoadAssetsWhenValidatingChangelists : 1 = true;
+
+	UPROPERTY(EditAnywhere, Config, Category="Data Validation")
+	TArray<FMaterialEditorValidationPlatform> MaterialValidationPlatforms;
 };
diff --git a/Engine/Plugins/Editor/DataValidation/Source/DataValidation/Public/EditorValidator_Material.h b/Engine/Plugins/Editor/DataValidation/Source/DataValidation/Public/EditorValidator_Material.h
new file mode 100644
--- /dev/null
+++ b/Engine/Plugins/Editor/DataValidation/Source/DataValidation/Public/EditorValidator_Material.h
@@ -0,0 +1,86 @@
+﻿// Copyright Epic Games, Inc. All Rights Reserved.
+
+#pragma once
+
+#include "EditorValidatorBase.h"
+#include "RHIShaderPlatform.h"
+#include "RHIFeatureLevel.h"
+#include "SceneTypes.h"
+
+#include "EditorValidator_Material.generated.h"
+
+class FText;
+class UObject;
+
+// Checks if UMaterial and UMaterialInstance are compiling on all validation shader platforms
+UCLASS()
+class UEditorValidator_Material : public UEditorValidatorBase
+{
+	GENERATED_BODY()
+
+public:
+	UEditorValidator_Material();
+
+protected:
+	virtual bool CanValidateAsset_Implementation(const FAssetData& InAssetData, UObject* InAsset, FDataValidationContext& InContext) const override;
+	virtual EDataValidationResult ValidateLoadedAsset_Implementation(const FAssetData& InAssetData, UObject* InAsset, FDataValidationContext& InContext) override;
+
+private:
+	struct FShaderValidationPlatform
+	{
+		EShaderPlatform ShaderPlatform;
+		ERHIFeatureLevel::Type FeatureLevel;
+		EMaterialQualityLevel::Type MaterialQualityLevel;
+	};
+
+	TArray<FShaderValidationPlatform> ValidationPlatforms;
+};
+
+USTRUCT()
+struct FMaterialEditorValidationShaderPlatform
+{
+	GENERATED_BODY()
+
+	UPROPERTY(EditAnywhere, config, Category="Data Validation")
+	FName Name = MaxRHIShaderPlatformName;
+
+	// if Name == MaxRHIShaderPlatformName, use GMaxRHIShaderPlatform instead
+	static constexpr FStringView MaxRHIShaderPlatformNameView = TEXTVIEW("GMaxRHIShaderPlatform");
+	static FName MaxRHIShaderPlatformName;
+
+	static void RegisterCustomPropertyTypeLayout();
+	static void UnregisterCustomPropertyTypeLayout();
+};
+
+UENUM()
+enum class EMaterialEditorValidationFeatureLevel : int32
+{
+	CurrentMaxFeatureLevel = ERHIFeatureLevel::Num + 1	UMETA(DisplayName="Current Max Feature Level"),
+	ES3_1 = ERHIFeatureLevel::ES3_1						UMETA(DisplayName="ES3.1"),
+	SM5 = ERHIFeatureLevel::SM5,
+	SM6 = ERHIFeatureLevel::SM6,
+};
+
+UENUM()
+enum class EMaterialEditorValidationQualityLevel : uint8
+{
+	Low = EMaterialQualityLevel::Low,
+	Medium = EMaterialQualityLevel::Medium,
+	High = EMaterialQualityLevel::High,
+	Epic = EMaterialQualityLevel::Epic,
+};
+
+USTRUCT()
+struct FMaterialEditorValidationPlatform
+{
+	GENERATED_USTRUCT_BODY()
+
+	UPROPERTY(EditAnywhere, config, Category="Data Validation")
+	FMaterialEditorValidationShaderPlatform ShaderPlatform = {};
+
+	UPROPERTY(EditAnywhere, config, Category="Data Validation")
+	EMaterialEditorValidationFeatureLevel FeatureLevel = EMaterialEditorValidationFeatureLevel::CurrentMaxFeatureLevel;
+
+	UPROPERTY(EditAnywhere, config, Category="Data Validation")
+	EMaterialEditorValidationQualityLevel MaterialQualityLevel = EMaterialEditorValidationQualityLevel::Epic;
+};
diff --git a/Engine/Source/Runtime/RHI/Private/DataDrivenShaderPlatformInfo.cpp b/Engine/Source/Runtime/RHI/Private/DataDrivenShaderPlatformInfo.cpp
--- a/Engine/Source/Runtime/RHI/Private/DataDrivenShaderPlatformInfo.cpp
+++ b/Engine/Source/Runtime/RHI/Private/DataDrivenShaderPlatformInfo.cpp
@@ -346,6 +346,9 @@ void FGenericDataDrivenShaderPlatformInfo::ParseDataDrivenShaderInfo(const FConf
 	GET_SECTION_BOOL_HELPER(bSupportsWave64);
 	GET_SECTION_BOOL_HELPER(bSupportsIndependentSamplers);
 	GET_SECTION_BOOL_HELPER(bSupportsWorkGraphs);
+#if WITH_EDITOR
+	GET_SECTION_BOOL_HELPER(bCanUsePreviewPlatformForMaterialValidation);
+#endif
 #undef GET_SECTION_BOOL_HELPER
 #undef GET_SECTION_INT_HELPER
 #undef GET_SECTION_SUPPORT_HELPER
diff --git a/Engine/Source/Runtime/RHI/Public/DataDrivenShaderPlatformInfo.h b/Engine/Source/Runtime/RHI/Public/DataDrivenShaderPlatformInfo.h
--- a/Engine/Source/Runtime/RHI/Public/DataDrivenShaderPlatformInfo.h
+++ b/Engine/Source/Runtime/RHI/Public/DataDrivenShaderPlatformInfo.h
@@ -129,6 +129,9 @@ class FGenericDataDrivenShaderPlatformInfo
 	uint32 bSupportsWave64 : 1;
 	uint32 bSupportsIndependentSamplers : 1;
 	uint32 bSupportsWorkGraphs : 1;
+#if WITH_EDITOR
+	uint32 bCanUsePreviewPlatformForMaterialValidation : 1;
+#endif
 
 	// NOTE: When adding fields, you must also add to ParseDataDrivenShaderInfo!
 	uint32 bContainsValidPlatformInfo : 1;
@@ -837,6 +840,14 @@ public:
 		return Infos[Platform].bSupportsWorkGraphs;
 	}
 
+#if WITH_EDITOR
+	static FORCEINLINE_DEBUGGABLE const bool CanUseForMaterialValidation(const FStaticShaderPlatform Platform)
+	{
+		check(IsValid(Platform));
+		return Infos[Platform].bCanUsePreviewPlatformForMaterialValidation || !Infos[Platform].bIsPreviewPlatform;
+	}
+#endif
+
 	static FORCEINLINE_DEBUGGABLE const bool IsValid(const FStaticShaderPlatform Platform)
 	{
 		return Infos[Platform].bContainsValidPlatformInfo;
