Adding enable config for UMaterial and UMaterialInstance asset validator

diff --git a/Engine/Plugins/Editor/DataValidation/Source/DataValidation/Private/EditorValidator_Material.cpp b/Engine/Plugins/Editor/DataValidation/Source/DataValidation/Private/EditorValidator_Material.cpp
--- a/Engine/Plugins/Editor/DataValidation/Source/DataValidation/Private/EditorValidator_Material.cpp
+++ b/Engine/Plugins/Editor/DataValidation/Source/DataValidation/Private/EditorValidator_Material.cpp
@@ -13,56 +13,59 @@
 UEditorValidator_Material::UEditorValidator_Material()
 	: Super()
 {
-	for (const FMaterialEditorValidationPlatform& Config: GetDefault<UDataValidationSettings>()->MaterialValidationPlatforms)
+	if (GetDefault<UDataValidationSettings>()->bEnableMaterialValidation)
 	{
-		FShaderValidationPlatform Platform = {};
-
-		bool bValidShaderPlatform = false;
-		if (Config.ShaderPlatform.Name == FMaterialEditorValidationShaderPlatform::MaxRHIShaderPlatformName)
-		{
-			Platform.ShaderPlatform = GMaxRHIShaderPlatform;
-			bValidShaderPlatform = true;
-		}
-		else
+		for (const FMaterialEditorValidationPlatform& Config: GetDefault<UDataValidationSettings>()->MaterialValidationPlatforms)
 		{
-			for (int32 ShaderPlatformIndex = 0; ShaderPlatformIndex < SP_NumPlatforms; ++ShaderPlatformIndex)
-			{
-				const EShaderPlatform ShaderPlatform = static_cast<EShaderPlatform>(ShaderPlatformIndex);
+			FShaderValidationPlatform Platform = {};
 
-				if (FDataDrivenShaderPlatformInfo::IsValid(ShaderPlatform)
-					&& FDataDrivenShaderPlatformInfo::CanUseForMaterialValidation(ShaderPlatform)
-					&& FDataDrivenShaderPlatformInfo::GetName(ShaderPlatform) == Config.ShaderPlatform.Name)
+			bool bValidShaderPlatform = false;
+			if (Config.ShaderPlatform.Name == FMaterialEditorValidationShaderPlatform::MaxRHIShaderPlatformName)
+			{
+				Platform.ShaderPlatform = GMaxRHIShaderPlatform;
+				bValidShaderPlatform = true;
+			}
+			else
+			{
+				for (int32 ShaderPlatformIndex = 0; ShaderPlatformIndex < SP_NumPlatforms; ++ShaderPlatformIndex)
 				{
-					Platform.ShaderPlatform = ShaderPlatform;
-					bValidShaderPlatform = true;
-					break;
+					const EShaderPlatform ShaderPlatform = static_cast<EShaderPlatform>(ShaderPlatformIndex);
+
+					if (FDataDrivenShaderPlatformInfo::IsValid(ShaderPlatform)
+						&& FDataDrivenShaderPlatformInfo::CanUseForMaterialValidation(ShaderPlatform)
+						&& FDataDrivenShaderPlatformInfo::GetName(ShaderPlatform) == Config.ShaderPlatform.Name)
+					{
+						Platform.ShaderPlatform = ShaderPlatform;
+						bValidShaderPlatform = true;
+						break;
+					}
 				}
 			}
-		}
 
-		if (!bValidShaderPlatform)
-		{
-			UE_LOG(LogContentValidation, Warning, TEXT("Material asset validation shader platform '%s' is not available, skipping."), *Config.ShaderPlatform.Name.ToString());
-			continue;
-		}
+			if (!bValidShaderPlatform)
+			{
+				UE_LOG(LogContentValidation, Warning, TEXT("Material asset validation shader platform '%s' is not available, skipping."), *Config.ShaderPlatform.Name.ToString());
+				continue;
+			}
 
-		switch (Config.FeatureLevel)
-		{
-		case EMaterialEditorValidationFeatureLevel::CurrentMaxFeatureLevel: Platform.FeatureLevel = GMaxRHIFeatureLevel; break; 
-		case EMaterialEditorValidationFeatureLevel::ES3_1: Platform.FeatureLevel = ERHIFeatureLevel::ES3_1; break;
-		case EMaterialEditorValidationFeatureLevel::SM5: Platform.FeatureLevel = ERHIFeatureLevel::SM5; break;
-		case EMaterialEditorValidationFeatureLevel::SM6: Platform.FeatureLevel = ERHIFeatureLevel::SM6; break;
-		}
+			switch (Config.FeatureLevel)
+			{
+			case EMaterialEditorValidationFeatureLevel::CurrentMaxFeatureLevel: Platform.FeatureLevel = GMaxRHIFeatureLevel; break; 
+			case EMaterialEditorValidationFeatureLevel::ES3_1: Platform.FeatureLevel = ERHIFeatureLevel::ES3_1; break;
+			case EMaterialEditorValidationFeatureLevel::SM5: Platform.FeatureLevel = ERHIFeatureLevel::SM5; break;
+			case EMaterialEditorValidationFeatureLevel::SM6: Platform.FeatureLevel = ERHIFeatureLevel::SM6; break;
+			}
 
-		switch (Config.MaterialQualityLevel)
-		{
-		case EMaterialEditorValidationQualityLevel::Low: Platform.MaterialQualityLevel = EMaterialQualityLevel::Low; break;
-		case EMaterialEditorValidationQualityLevel::Medium: Platform.MaterialQualityLevel = EMaterialQualityLevel::Medium; break;
-		case EMaterialEditorValidationQualityLevel::High: Platform.MaterialQualityLevel = EMaterialQualityLevel::High; break;
-		case EMaterialEditorValidationQualityLevel::Epic: Platform.MaterialQualityLevel = EMaterialQualityLevel::Epic; break;
-		}
+			switch (Config.MaterialQualityLevel)
+			{
+			case EMaterialEditorValidationQualityLevel::Low: Platform.MaterialQualityLevel = EMaterialQualityLevel::Low; break;
+			case EMaterialEditorValidationQualityLevel::Medium: Platform.MaterialQualityLevel = EMaterialQualityLevel::Medium; break;
+			case EMaterialEditorValidationQualityLevel::High: Platform.MaterialQualityLevel = EMaterialQualityLevel::High; break;
+			case EMaterialEditorValidationQualityLevel::Epic: Platform.MaterialQualityLevel = EMaterialQualityLevel::Epic; break;
+			}
 
-		ValidationPlatforms.Add(Platform);
+			ValidationPlatforms.Add(Platform);
+		}
 	}
 }
 
diff --git a/Engine/Plugins/Editor/DataValidation/Source/DataValidation/Public/DataValidationSettings.h b/Engine/Plugins/Editor/DataValidation/Source/DataValidation/Public/DataValidationSettings.h
--- a/Engine/Plugins/Editor/DataValidation/Source/DataValidation/Public/DataValidationSettings.h
+++ b/Engine/Plugins/Editor/DataValidation/Source/DataValidation/Public/DataValidationSettings.h
@@ -28,5 +28,8 @@ public:
 	uint32 bLoadAssetsWhenValidatingChangelists : 1 = true;
 
 	UPROPERTY(EditAnywhere, Config, Category="Data Validation")
+	bool bEnableMaterialValidation = true;
+
+	UPROPERTY(EditAnywhere, Config, Category="Data Validation", meta=(EditCondition="bEnableMaterialValidation"))
 	TArray<FMaterialEditorValidationPlatform> MaterialValidationPlatforms;
 };
