Don't run material asset validator in commandlet

diff --git a/Engine/Plugins/Editor/DataValidation/Source/DataValidation/Private/EditorValidator_Material.cpp b/Engine/Plugins/Editor/DataValidation/Source/DataValidation/Private/EditorValidator_Material.cpp
--- a/Engine/Plugins/Editor/DataValidation/Source/DataValidation/Private/EditorValidator_Material.cpp
+++ b/Engine/Plugins/Editor/DataValidation/Source/DataValidation/Private/EditorValidator_Material.cpp
@@ -71,9 +71,17 @@ UEditorValidator_Material::UEditorValidator_Material()
 
 bool UEditorValidator_Material::CanValidateAsset_Implementation(const FAssetData& AssetData, UObject* InAsset, FDataValidationContext& InContext) const
 {
-	const bool bAnyValidationPlatforms = !ValidationPlatforms.IsEmpty();
-	const bool bShouldAllowFullValidation = InContext.GetValidationUsecase() != EDataValidationUsecase::Save;
-	return bAnyValidationPlatforms && bShouldAllowFullValidation && (Cast<UMaterial>(InAsset) || Cast<UMaterialInstance>(InAsset));
+	if (InContext.GetValidationUsecase() == EDataValidationUsecase::Save || InContext.GetValidationUsecase() == EDataValidationUsecase::Commandlet)
+	{
+		return false;
+	}
+
+	if (ValidationPlatforms.IsEmpty())
+	{
+		return false;
+	}
+
+	return Cast<UMaterial>(InAsset) || Cast<UMaterialInstance>(InAsset);
 }
 
 EDataValidationResult UEditorValidator_Material::ValidateLoadedAsset_Implementation(const FAssetData& AssetData, UObject* InAsset, FDataValidationContext& InContext)
