Duplicating UMaterial and UMaterialAsset to fix stale shader compilation jobs

diff --git a/Engine/Plugins/Editor/DataValidation/Source/DataValidation/Private/EditorValidator_Material.cpp b/Engine/Plugins/Editor/DataValidation/Source/DataValidation/Private/EditorValidator_Material.cpp
--- a/Engine/Plugins/Editor/DataValidation/Source/DataValidation/Private/EditorValidator_Material.cpp
+++ b/Engine/Plugins/Editor/DataValidation/Source/DataValidation/Private/EditorValidator_Material.cpp
@@ -81,34 +81,57 @@ bool UEditorValidator_Material::CanValidateAsset_Implementation(const FAssetData
 		return false;
 	}
 
-	return Cast<UMaterial>(InAsset) || Cast<UMaterialInstance>(InAsset);
+	const UMaterial* OriginalMaterial = Cast<UMaterial>(InAsset);
+	if (OriginalMaterial != nullptr) // we can always validate a UMaterial
+	{
+		return true;
+	}
+
+	const UMaterialInstance* OriginalMaterialInstance = Cast<UMaterialInstance>(InAsset);
+	if (OriginalMaterialInstance != nullptr)
+	{
+		FMaterialInheritanceChain Chain;
+		OriginalMaterialInstance->GetMaterialInheritanceChain(Chain);
+
+		for (const UMaterialInstance* MaterialInstance: Chain.MaterialInstances)
+		{
+			if (MaterialInstance->HasStaticParameters())
+			{
+				// only validate UMaterialInstance if it influences shader compilation
+				return true;
+			}
+		}
+	}
+
+	return false;
 }
 
 EDataValidationResult UEditorValidator_Material::ValidateLoadedAsset_Implementation(const FAssetData& AssetData, UObject* InAsset, FDataValidationContext& InContext)
 {
-	UMaterialInstance* MaterialInstance = Cast<UMaterialInstance>(InAsset);
+	UMaterialInstance* OriginalMaterialInstance = Cast<UMaterialInstance>(InAsset);
+	UMaterial* OriginalMaterial = OriginalMaterialInstance ? OriginalMaterialInstance->GetMaterial() : Cast<UMaterial>(InAsset);
 
-	UMaterial* Material = MaterialInstance ? MaterialInstance->GetMaterial() : Cast<UMaterial>(InAsset);
-	check(Material);
+	UMaterialInstance* MaterialInstance = DuplicateMaterialInstance(OriginalMaterialInstance);
+	UMaterial* Material = MaterialInstance ? MaterialInstance->GetMaterial() : DuplicateMaterial(OriginalMaterial);
+	if (!ensureAlways(OriginalMaterial) || !ensureAlways(Material))
+	{
+		return EDataValidationResult::NotValidated;
+	}
 
 	TArray<FMaterialResource*> Resources;
-	bool bAnyResources = false;
 
 	for (const FShaderValidationPlatform& ValidationPlatform: ValidationPlatforms)
 	{
 		FMaterialResource* CurrentResource = FindOrCreateMaterialResource(Resources, Material, MaterialInstance, ValidationPlatform.FeatureLevel, ValidationPlatform.MaterialQualityLevel);
-		check(CurrentResource);
 
-		if (CurrentResource)
+		if (ensure(CurrentResource))
 		{
 			CurrentResource->CacheShaders(ValidationPlatform.ShaderPlatform);
-			bAnyResources = true;
 		}
 	}
 
-	if (!bAnyResources)
+	if (Resources.IsEmpty())
 	{
-		FMaterial::DeferredDeleteArray(Resources);
 		return EDataValidationResult::NotValidated;
 	}
 
@@ -130,6 +153,82 @@ EDataValidationResult UEditorValidator_Material::ValidateLoadedAsset_Implementat
 	return !bCompileErrors ? EDataValidationResult::Valid : EDataValidationResult::Invalid;
 }
 
+UMaterial* UEditorValidator_Material::DuplicateMaterial(UMaterial* OriginalMaterial)
+{
+	if (!OriginalMaterial)
+	{
+		return nullptr;
+	}
+
+	return static_cast<UMaterial*>(StaticDuplicateObject(OriginalMaterial, GetTransientPackage(), NAME_None, ~RF_Standalone, UValidationMaterial::StaticClass()));
+}
+
+UMaterialInstance* UEditorValidator_Material::DuplicateMaterialInstance(UMaterialInstance* OriginalMaterialInstance)
+{
+	if (!OriginalMaterialInstance)
+	{
+		return nullptr;
+	}
+
+	TArray<UMaterialInstance*> DuplicatedMaterialInstances;
+
+	FMaterialInheritanceChain Chain;
+	OriginalMaterialInstance->GetMaterialInheritanceChain(Chain);
+
+	for (const UMaterialInstance* MaterialInstance: Chain.MaterialInstances)
+	{
+		// only duplicate material instances that might influence compilation
+		if (!MaterialInstance->HasStaticParameters())
+		{
+			continue;
+		}
+
+		UMaterialInstance* DuplicatedMaterialInstance = Cast<UMaterialInstance>(StaticDuplicateObject(MaterialInstance, GetTransientPackage(), NAME_None, ~RF_Standalone, MaterialInstance->GetClass()));
+
+		DuplicatedMaterialInstances.Add(DuplicatedMaterialInstance);
+	}
+
+	// should be caught by CanValidateAsset_Implementation
+	if (ensureAlways(DuplicatedMaterialInstances.Num() > 0))
+	{
+		UMaterial* DuplicatedMaterial = DuplicateMaterial(OriginalMaterialInstance->GetMaterial());
+
+		for (int32 i = 0; i < DuplicatedMaterialInstances.Num(); ++i)
+		{
+			if (i + 1 < DuplicatedMaterialInstances.Num())
+			{
+				DuplicatedMaterialInstances[i]->Parent = DuplicatedMaterialInstances[i + 1];
+			}
+			else
+			{
+				DuplicatedMaterialInstances[i]->Parent = DuplicatedMaterial;
+			}
+		}
+
+		return DuplicatedMaterialInstances[0];
+	}
+	else
+	{
+		return nullptr;
+	}
+}
+
+class FValidationMaterial : public FMaterialResource
+{
+public:
+	FValidationMaterial() = default;
+	virtual ~FValidationMaterial() override = default;
+
+	virtual bool IsPersistent() const override { return false; }
+	virtual FString GetAssetName() const override { return FString::Printf(TEXT("Validation:%s"), *FMaterialResource::GetAssetName()); }
+	virtual bool IsPreview() const override { return true; }
+};
+
+FMaterialResource* UValidationMaterial::AllocateResource()
+{
+	return new FValidationMaterial();
+}
+
 class FMaterialEditorValidationPlatformCustomization : public IPropertyTypeCustomization
 {
 public:
diff --git a/Engine/Plugins/Editor/DataValidation/Source/DataValidation/Public/EditorValidator_Material.h b/Engine/Plugins/Editor/DataValidation/Source/DataValidation/Public/EditorValidator_Material.h
--- a/Engine/Plugins/Editor/DataValidation/Source/DataValidation/Public/EditorValidator_Material.h
+++ b/Engine/Plugins/Editor/DataValidation/Source/DataValidation/Public/EditorValidator_Material.h
@@ -6,6 +6,8 @@
 #include "RHIShaderPlatform.h"
 #include "RHIFeatureLevel.h"
 #include "SceneTypes.h"
+#include "Materials/Material.h"
+#include "Materials/MaterialInstance.h"
 
 #include "EditorValidator_Material.generated.h"
 
@@ -34,6 +36,19 @@ private:
 	};
 
 	TArray<FShaderValidationPlatform> ValidationPlatforms;
+
+	static UMaterial* DuplicateMaterial(UMaterial* OriginalMaterial);
+	static UMaterialInstance* DuplicateMaterialInstance(UMaterialInstance* OriginalMaterialInstance);
+};
+
+UCLASS()
+class UValidationMaterial : public UMaterial
+{
+	GENERATED_BODY()
+
+public:
+	virtual FMaterialResource* AllocateResource() override;
+	virtual bool IsAsset() const override { return false; }
 };
 
 USTRUCT()
