Fixing DataValidation module shutdown memory corruption

diff --git a/Engine/Plugins/Editor/DataValidation/Source/DataValidation/Private/EditorValidator_Material.cpp b/Engine/Plugins/Editor/DataValidation/Source/DataValidation/Private/EditorValidator_Material.cpp
--- a/Engine/Plugins/Editor/DataValidation/Source/DataValidation/Private/EditorValidator_Material.cpp
+++ b/Engine/Plugins/Editor/DataValidation/Source/DataValidation/Private/EditorValidator_Material.cpp
@@ -369,10 +369,12 @@ private:
 };
 
 FName FMaterialEditorValidationShaderPlatform::MaxRHIShaderPlatformName = MaxRHIShaderPlatformNameView.GetData();
+FName FMaterialEditorValidationShaderPlatform::CustomPropertyTypeLayoutName;
 
 void FMaterialEditorValidationShaderPlatform::RegisterCustomPropertyTypeLayout()
 {
 	FPropertyEditorModule& PropertyModule = FModuleManager::LoadModuleChecked<FPropertyEditorModule>("PropertyEditor");
+	CustomPropertyTypeLayoutName = StaticStruct()->GetFName(); 
 	PropertyModule.RegisterCustomPropertyTypeLayout(
 		StaticStruct()->GetFName(),
 		FOnGetPropertyTypeCustomizationInstance::CreateLambda([](){return MakeShared<FMaterialEditorValidationPlatformCustomization>();})
@@ -381,6 +383,10 @@ void FMaterialEditorValidationShaderPlatform::RegisterCustomPropertyTypeLayout()
 
 void FMaterialEditorValidationShaderPlatform::UnregisterCustomPropertyTypeLayout()
 {
-	FPropertyEditorModule& PropertyModule = FModuleManager::LoadModuleChecked<FPropertyEditorModule>("PropertyEditor");
-	PropertyModule.UnregisterCustomPropertyTypeLayout(StaticStruct()->GetFName());
+	FPropertyEditorModule* PropertyModule = FModuleManager::GetModulePtr<FPropertyEditorModule>("PropertyEditor");
+	if (PropertyModule)
+	{
+		// StaticStruct()->GetFName() is not available during engine shutdown as UObjects were already destroyed
+		PropertyModule->UnregisterCustomPropertyTypeLayout(CustomPropertyTypeLayoutName);
+	}
 }
diff --git a/Engine/Plugins/Editor/DataValidation/Source/DataValidation/Public/EditorValidator_Material.h b/Engine/Plugins/Editor/DataValidation/Source/DataValidation/Public/EditorValidator_Material.h
--- a/Engine/Plugins/Editor/DataValidation/Source/DataValidation/Public/EditorValidator_Material.h
+++ b/Engine/Plugins/Editor/DataValidation/Source/DataValidation/Public/EditorValidator_Material.h
@@ -63,6 +63,7 @@ struct FMaterialEditorValidationShaderPlatform
 	static constexpr FStringView MaxRHIShaderPlatformNameView = TEXTVIEW("GMaxRHIShaderPlatform");
 	static FName MaxRHIShaderPlatformName;
 
+	static FName CustomPropertyTypeLayoutName;
 	static void RegisterCustomPropertyTypeLayout();
 	static void UnregisterCustomPropertyTypeLayout();
 };
