Added a missed use of the validation settings when checking if an asset should be validated or not.

diff --git a/Engine/Plugins/Editor/DataValidation/Source/DataValidation/Private/EditorValidatorSubsystem.cpp b/Engine/Plugins/Editor/DataValidation/Source/DataValidation/Private/EditorValidatorSubsystem.cpp
--- a/Engine/Plugins/Editor/DataValidation/Source/DataValidation/Private/EditorValidatorSubsystem.cpp
+++ b/Engine/Plugins/Editor/DataValidation/Source/DataValidation/Private/EditorValidatorSubsystem.cpp
@@ -137,7 +137,7 @@ bool UEditorValidatorSubsystem::ShouldValidateAsset(
 		return false;
 	}
 
-	if (IsPathExcludedFromValidation(AssetPackageNameView))
+	if (Settings.bSkipExcludedDirectories && IsPathExcludedFromValidation(AssetPackageNameView))
 	{
 		return false;
 	}
