Material validator review fixups

diff --git a/Engine/Plugins/Editor/DataValidation/Source/DataValidation/Private/EditorValidator_Material.cpp b/Engine/Plugins/Editor/DataValidation/Source/DataValidation/Private/EditorValidator_Material.cpp
--- a/Engine/Plugins/Editor/DataValidation/Source/DataValidation/Private/EditorValidator_Material.cpp
+++ b/Engine/Plugins/Editor/DataValidation/Source/DataValidation/Private/EditorValidator_Material.cpp
@@ -82,13 +82,14 @@ bool UEditorValidator_Material::CanValidateAsset_Implementation(const FAssetData
 	}
 
 	const UMaterial* OriginalMaterial = Cast<UMaterial>(InAsset);
-	if (OriginalMaterial != nullptr) // we can always validate a UMaterial
+	if (OriginalMaterial)
 	{
+		// We want to validate every UMaterial.
 		return true;
 	}
 
 	const UMaterialInstance* OriginalMaterialInstance = Cast<UMaterialInstance>(InAsset);
-	if (OriginalMaterialInstance != nullptr)
+	if (OriginalMaterialInstance)
 	{
 		FMaterialInheritanceChain Chain;
 		OriginalMaterialInstance->GetMaterialInheritanceChain(Chain);
@@ -97,7 +98,7 @@ bool UEditorValidator_Material::CanValidateAsset_Implementation(const FAssetData
 		{
 			if (MaterialInstance->HasStaticParameters())
 			{
-				// only validate UMaterialInstance if it influences shader compilation
+				// We want to validate UMaterialInstance since it has a static parameter that could influence generated shader code.
 				return true;
 			}
 		}
@@ -108,9 +109,11 @@ bool UEditorValidator_Material::CanValidateAsset_Implementation(const FAssetData
 
 EDataValidationResult UEditorValidator_Material::ValidateLoadedAsset_Implementation(const FAssetData& AssetData, UObject* InAsset, FDataValidationContext& InContext)
 {
-	UMaterialInstance* OriginalMaterialInstance = Cast<UMaterialInstance>(InAsset);
+	UMaterialInstance* OriginalMaterialInstance = Cast<UMaterialInstance>(InAsset); // Can be null
 	UMaterial* OriginalMaterial = OriginalMaterialInstance ? OriginalMaterialInstance->GetMaterial() : Cast<UMaterial>(InAsset);
 
+	// Right now calling CacheShaders multiple times on the same UMaterial(Instance) will not work as we don't clean up stale compilation jobs.
+	// So to work around this we create a duplicate assets and compile cache shaders for them instead.
 	UMaterialInstance* MaterialInstance = DuplicateMaterialInstance(OriginalMaterialInstance);
 	UMaterial* Material = MaterialInstance ? MaterialInstance->GetMaterial() : DuplicateMaterial(OriginalMaterial);
 	if (!ensureAlways(OriginalMaterial) || !ensureAlways(Material))
diff --git a/Engine/Plugins/Editor/DataValidation/Source/DataValidation/Public/EditorValidator_Material.h b/Engine/Plugins/Editor/DataValidation/Source/DataValidation/Public/EditorValidator_Material.h
--- a/Engine/Plugins/Editor/DataValidation/Source/DataValidation/Public/EditorValidator_Material.h
+++ b/Engine/Plugins/Editor/DataValidation/Source/DataValidation/Public/EditorValidator_Material.h
@@ -37,6 +37,7 @@ private:
 
 	TArray<FShaderValidationPlatform> ValidationPlatforms;
 
+	// Creates a duplicate asset in transient package, will return nullptr if nullptr is passed.
 	static UMaterial* DuplicateMaterial(UMaterial* OriginalMaterial);
 	static UMaterialInstance* DuplicateMaterialInstance(UMaterialInstance* OriginalMaterialInstance);
 };
