Open message log for warnings/errors from standard changelist presubmit validation.

diff --git a/Engine/Plugins/Editor/DataValidation/Source/DataValidation/Private/EditorValidatorSubsystem.cpp b/Engine/Plugins/Editor/DataValidation/Source/DataValidation/Private/EditorValidatorSubsystem.cpp
--- a/Engine/Plugins/Editor/DataValidation/Source/DataValidation/Private/EditorValidatorSubsystem.cpp
+++ b/Engine/Plugins/Editor/DataValidation/Source/DataValidation/Private/EditorValidatorSubsystem.cpp
@@ -800,6 +800,7 @@ void UEditorValidatorSubsystem::ValidateChangelistPreSubmit(
 	Settings.ValidationUsecase = EDataValidationUsecase::PreSubmit;
 	Settings.bLoadAssetsForValidation = GetDefault<UDataValidationSettings>()->bLoadAssetsWhenValidatingChangelists;
 	Settings.MessageLogPageTitle = FText::Format(LOCTEXT("MessageLogPageTitle.ValidateChangelist", "Changelist Validation: {0}"), FText::FromString(InChangelist->GetIdentifier()));
+	Settings.ShowMessageLogSeverity = { EMessageSeverity::Warning };
 
 	FValidateAssetsResults Results;
 	OutResult = ValidateChangelist(Changelist, Settings, Results);
@@ -929,6 +930,11 @@ EDataValidationResult UEditorValidatorSubsystem::ValidateChangelistsInternal(
 	Result = CombineDataValidationResults(Result, AssetResult);
 	LogAssetValidationSummary(DataValidationLog, Settings, Result, OutResults);
 
+	if (const EMessageSeverity::Type* Severity = Settings.ShowMessageLogSeverity.GetPtrOrNull())
+	{
+		DataValidationLog.Open(*Severity, false);
+	}
+
 	return Result;
 }
 
