Allow editor validator subsystem to load external actors with maps for validation.
Controlled by FValidateAssetsSettings::bLoadExternalObjectsForValidation and only performed when validating "assets" rather than "objects" in the validation subsystem API functions.
Add bUnloadAssetsLoadedForValidation to allow unloading assets post-validation (e.g. may be useful when validating many assets in the editor).

#tests Tested validating a map with an actor which always fails validation in editor and from submit tool.

diff --git a/Engine/Plugins/Editor/DataValidation/Source/DataValidation/Private/DataValidationTestActor.cpp b/Engine/Plugins/Editor/DataValidation/Source/DataValidation/Private/DataValidationTestActor.cpp
new file mode 100644
--- /dev/null
+++ b/Engine/Plugins/Editor/DataValidation/Source/DataValidation/Private/DataValidationTestActor.cpp
@@ -0,0 +1,19 @@
+// Copyright Epic Games, Inc. All Rights Reserved.
+
+#include "DataValidationTestActor.h"
+
+#include "Misc/DataValidation.h"
+
+#define LOCTEXT_NAMESPACE "AssetValidation"
+
+EDataValidationResult ADataValidationTestActor::IsDataValid(FDataValidationContext& Context) const 
+{
+    if (!bPassValidation)
+    {
+        Context.AddError(LOCTEXT("bPassValidationFalse", "bPassValidation is false"));
+        return EDataValidationResult::Invalid;
+    }
+    return EDataValidationResult::Valid;
+}
+
+#undef LOCTEXT_NAMESPACE 
\ No newline at end of file
diff --git a/Engine/Plugins/Editor/DataValidation/Source/DataValidation/Private/DataValidationTestActor.h b/Engine/Plugins/Editor/DataValidation/Source/DataValidation/Private/DataValidationTestActor.h
new file mode 100644
--- /dev/null
+++ b/Engine/Plugins/Editor/DataValidation/Source/DataValidation/Private/DataValidationTestActor.h
@@ -0,0 +1,22 @@
+// Copyright Epic Games, Inc. All Rights Reserved.
+
+#pragma once
+
+#include "GameFramework/Info.h"
+
+#include "DataValidationTestActor.generated.h"
+
+UCLASS()
+class ADataValidationTestActor : public AInfo
+{
+    GENERATED_BODY()
+
+public:
+    UPROPERTY(EditAnywhere, BlueprintReadWrite, Category="Validation")
+    bool bPassValidation = false;
+
+#if WITH_EDITOR
+    EDataValidationResult IsDataValid(FDataValidationContext& Context) const override;
+#endif
+
+};
\ No newline at end of file
diff --git a/Engine/Plugins/Editor/DataValidation/Source/DataValidation/Private/EditorValidatorSubsystem.cpp b/Engine/Plugins/Editor/DataValidation/Source/DataValidation/Private/EditorValidatorSubsystem.cpp
--- a/Engine/Plugins/Editor/DataValidation/Source/DataValidation/Private/EditorValidatorSubsystem.cpp
+++ b/Engine/Plugins/Editor/DataValidation/Source/DataValidation/Private/EditorValidatorSubsystem.cpp
@@ -19,6 +19,7 @@
 #include "Engine/Level.h"
 #include "IDirectoryWatcher.h"
 #include "ISourceControlModule.h"
+#include "LoadExternalObjectsForValidation.h"
 #include "Logging/MessageLog.h"
 #include "Misc/App.h"
 #include "Misc/DataValidation.h"
@@ -26,6 +27,7 @@
 #include "Misc/PackageName.h"
 #include "Misc/ScopedSlowTask.h"
 #include "Misc/ScopeExit.h"
+#include "PackageTools.h"
 #include "UObject/ICookInfo.h"
 #include "UObject/TopLevelAssetPath.h"
 
@@ -451,6 +453,8 @@ EDataValidationResult UEditorValidatorSubsystem::ValidateAssetsInternal(
 	OutResults.NumRequested = AssetDataList.Num();
 	
 	EDataValidationResult Result = EDataValidationResult::NotValidated;
+	// Loaded assets ought to have the standalone flag and so not be garbage collected, but keep weak pointers here for memory safety
+	TArray<TWeakObjectPtr<UPackage>> PackagesToUnload;	
 
 	// Now add to map or update as needed
 	for (const FAssetData& Data : AssetDataList)
@@ -508,7 +512,15 @@ EDataValidationResult UEditorValidatorSubsystem::ValidateAssetsInternal(
 		{
 			UE_LOG(LogContentValidation, Log, TEXT("Loading asset %s for validation"), *Data.ToSoftObjectPath().ToString());
 			UE::DataValidation::FScopedLogMessageGatherer LogGatherer(InSettings.bCaptureAssetLoadLogs);
-			LoadedAsset = Data.GetAsset(); // Do not load external objects, validators should load external objects that they want via GetAssociatedExternalObjects in the validation context 
+			
+			// Don't pass the flag to load all external objects. Later we'll load external objects which are also part of the to-validate
+			// set if the settings allow for it.
+			LoadedAsset = Data.GetAsset(); 
+
+			if (InSettings.bUnloadAssetsLoadedForValidation && LoadedAsset)
+			{
+				PackagesToUnload.Emplace(LoadedAsset->GetPackage());
+			}
 
 			WaitForAssetCompilationIfNecessary(InSettings.ValidationUsecase);
 			
@@ -537,6 +549,8 @@ EDataValidationResult UEditorValidatorSubsystem::ValidateAssetsInternal(
 
 		if (LoadedAsset)
 		{
+			UE::DataValidation::FScopedLoadExternalObjects ExternalObjectsLoader(LoadedAsset, ValidationContext, InSettings.bLoadExternalObjectsForValidation);
+
 			UE::DataValidation::FScopedLogMessageGatherer LogGatherer(InSettings.bCaptureLogsDuringValidation);
 			AssetResult = IsObjectValidWithContext(LoadedAsset, ValidationContext);
 			
@@ -648,6 +662,13 @@ EDataValidationResult UEditorValidatorSubsystem::ValidateAssetsInternal(
 		DataValidationLog.Flush();
 		
 		Result = CombineDataValidationResults(Result, AssetResult);
+
+		if (PackagesToUnload.Num() > 0)
+		{
+			TArray<UPackage*> LocalPackages;
+			CopyFromWeakArray(LocalPackages, PackagesToUnload);
+			UPackageTools::UnloadPackages(LocalPackages);
+		}
 	}
 
 	// Broadcast now that we're complete so other systems can go back to their previous state.
diff --git a/Engine/Plugins/Editor/DataValidation/Source/DataValidation/Private/LoadExternalObjectsForValidation.cpp b/Engine/Plugins/Editor/DataValidation/Source/DataValidation/Private/LoadExternalObjectsForValidation.cpp
new file mode 100644
--- /dev/null
+++ b/Engine/Plugins/Editor/DataValidation/Source/DataValidation/Private/LoadExternalObjectsForValidation.cpp
@@ -0,0 +1,124 @@
+// Copyright Epic Games, Inc. All Rights Reserved.
+
+#include "LoadExternalObjectsForValidation.h"
+
+#include "AssetRegistry/AssetData.h"
+#include "Engine/World.h"
+#include "Misc/DataValidation.h"
+#include "Misc/PathViews.h"
+#include "WorldPartition/ContentBundle/ContentBundlePaths.h"
+#include "WorldPartition/DataLayer/ExternalDataLayerAsset.h"
+#include "WorldPartition/DataLayer/ExternalDataLayerHelper.h"
+#include "WorldPartition/WorldPartition.h"
+#include "WorldPartition/WorldPartitionActorDescUtils.h"
+
+namespace UE::DataValidation
+{
+
+FScopedLoadExternalObjects::FScopedLoadExternalObjects(UObject* InAsset, FDataValidationContext& InContext, bool bEnabled)
+{
+    if (!bEnabled)
+    {
+        return;
+    }
+
+    if (UWorld* World = Cast<UWorld>(InAsset))
+    {
+        auto RegisterContainerToValidate = [this](UWorld* InWorld, FName InContainerPackageName, FActorDescContainerInstanceCollection& OutRegisteredContainers, const FGuid& InContentBundleGuid = FGuid(), const UExternalDataLayerAsset* InExternalDataLayerAsset = nullptr)
+        {
+            if (OutRegisteredContainers.Contains(InContainerPackageName))
+            {
+                return;
+            }
+
+            UActorDescContainerInstance* ContainerInstance = nullptr;
+            if (InWorld != nullptr)
+            {
+                // World is Loaded reuse the ActorDescContainer of the Content Bundle
+                ContainerInstance = InWorld->GetWorldPartition()->FindContainer(InContainerPackageName);
+            }
+
+            // Even if world is valid, its world partition is not necessarily initialized
+            if (!ContainerInstance)
+            {
+                // Find in memory failed, load the ActorDescContainer
+                ContainerInstance = NewObject<UActorDescContainerInstance>();
+                UActorDescContainerInstance::FInitializeParams InitializeParams(InContainerPackageName);
+                InitializeParams.ContentBundleGuid = InContentBundleGuid;
+                InitializeParams.ExternalDataLayerAsset = InExternalDataLayerAsset;
+                ContainerInstance->Initialize(InitializeParams);
+                ContainersToUninit.Emplace(ContainerInstance);
+            }
+            else
+            {
+                check(ContainerInstance->GetContentBundleGuid() == InContentBundleGuid);
+                check(ContainerInstance->GetExternalDataLayerAsset() == InExternalDataLayerAsset);
+            }
+
+            OutRegisteredContainers.AddContainer(ContainerInstance);
+        };
+
+        TSet<FSoftObjectPath> ProcessedExternalDataLayers;
+        FString MapPackageName = World->GetPackage()->GetName();
+        TConstArrayView<FAssetData> ActorsData = InContext.GetAssociatedExternalObjects();
+        FActorDescContainerInstanceCollection ContainersToValidate;
+        for (const FAssetData& ActorData : ActorsData)
+        {
+            FString ActorPackagePath = ActorData.PackagePath.ToString();
+            if (ContentBundlePaths::IsAContentBundleExternalActorPackagePath(ActorPackagePath))
+            {
+                FStringView ContentBundleMountPoint = FPathViews::GetMountPointNameFromPath(ActorPackagePath);
+                FGuid ContentBundleGuid = ContentBundlePaths::GetContentBundleGuidFromExternalActorPackagePath(ActorPackagePath);
+                
+                FString ContentBundleContainerPackagePath;
+                verify(ContentBundlePaths::BuildActorDescContainerPackagePath(FString(ContentBundleMountPoint), ContentBundleGuid, MapPackageName, ContentBundleContainerPackagePath));
+
+                RegisterContainerToValidate(World, FName(*ContentBundleContainerPackagePath), ContainersToValidate, ContentBundleGuid);
+            }
+            else
+            {
+                if (TUniquePtr<FWorldPartitionActorDesc> ActorDesc = FWorldPartitionActorDescUtils::GetActorDescriptorFromAssetData(ActorData))
+                {
+                    bool bIsAlreadyInSet = false;
+                    const FSoftObjectPath& ExternalDataLayerPath = ActorDesc->GetExternalDataLayerAsset();
+                    if (ExternalDataLayerPath.IsValid())
+                    {
+                        ProcessedExternalDataLayers.Add(ExternalDataLayerPath, &bIsAlreadyInSet);
+                        if (!bIsAlreadyInSet)
+                        {
+                            if (const UExternalDataLayerAsset* ExternalDataLayerAsset = Cast<UExternalDataLayerAsset>(ExternalDataLayerPath.TryLoad()))
+                            {
+                                const FString EDLContainerPackagePath = FExternalDataLayerHelper::GetExternalDataLayerLevelRootPath(ExternalDataLayerAsset, MapPackageName);
+                                RegisterContainerToValidate(World, FName(*EDLContainerPackagePath), ContainersToValidate, FGuid(), ExternalDataLayerAsset);
+                            }
+                        }
+                    }
+                    else
+                    {
+                        RegisterContainerToValidate(World, World->GetPackage()->GetFName(), ContainersToValidate);
+                    }
+                }
+            }
+        }
+        
+        for (const FAssetData& ActorData : ActorsData)
+        {
+            if (const FWorldPartitionActorDescInstance* ActorDesc = ContainersToValidate.GetActorDescInstanceByPath(ActorData.AssetName.ToString()))
+            {
+                ActorRefs.Add(FWorldPartitionReference(&ContainersToValidate, ActorDesc->GetGuid()));
+            }
+        }
+    }
+}
+
+FScopedLoadExternalObjects::~FScopedLoadExternalObjects()
+{
+    // Explicitly release actors before uninitializing containers
+    ActorRefs.Empty();
+	for (const TStrongObjectPtr<UActorDescContainerInstance>& ContainerToUninit : ContainersToUninit)
+	{
+		ContainerToUninit->Uninitialize();
+	}
+    ContainersToUninit.Empty();
+}
+}
\ No newline at end of file
diff --git a/Engine/Plugins/Editor/DataValidation/Source/DataValidation/Private/LoadExternalObjectsForValidation.h b/Engine/Plugins/Editor/DataValidation/Source/DataValidation/Private/LoadExternalObjectsForValidation.h
new file mode 100644
--- /dev/null
+++ b/Engine/Plugins/Editor/DataValidation/Source/DataValidation/Private/LoadExternalObjectsForValidation.h
@@ -0,0 +1,28 @@
+// Copyright Epic Games, Inc. All Rights Reserved.
+
+#pragma once
+
+#include "UObject/StrongObjectPtrTemplates.h"
+#include "WorldPartition/WorldPartitionHandle.h"
+
+class FDataValidationContext;
+class UActorDescContainerInstance;
+class UObject;
+
+namespace UE::DataValidation
+{
+    // Utility to load external objects (e.g. actors) associated with an asset (e.g. a map) for validation so that they are 
+    // available to validators inspecting the main asset. 
+    // Should only be used on the stack. 
+    struct FScopedLoadExternalObjects
+    {
+        FScopedLoadExternalObjects(UObject* InAsset, FDataValidationContext& InContext, bool bEnabled);
+        ~FScopedLoadExternalObjects();
+ 
+    private:
+        // Hard refs keeping actors loaded for the same lifetime as this object
+        TArray<FWorldPartitionReference> ActorRefs;
+        // Containers used for actor loading which must be explicitly uninitialized at the end of this object's lifetime
+        TArray<TStrongObjectPtr<UActorDescContainerInstance>> ContainersToUninit;
+    };
+}
\ No newline at end of file
diff --git a/Engine/Plugins/Editor/DataValidation/Source/DataValidation/Public/EditorValidatorSubsystem.h b/Engine/Plugins/Editor/DataValidation/Source/DataValidation/Public/EditorValidatorSubsystem.h
--- a/Engine/Plugins/Editor/DataValidation/Source/DataValidation/Public/EditorValidatorSubsystem.h
+++ b/Engine/Plugins/Editor/DataValidation/Source/DataValidation/Public/EditorValidatorSubsystem.h
@@ -138,7 +138,18 @@ struct DATAVALIDATION_API FValidateAssetsSettings
 	/** If false, unloaded assets will get skipped from validation. */
 	UPROPERTY(EditAnywhere, BlueprintReadWrite, Category = "Validation")
 	bool bLoadAssetsForValidation = true;
+
+	/** If true, will attempt to unload assets which were previously unloaded, and loaded for validation. */
+	UPROPERTY(EditAnywhere, BlueprintReadWrite, Category = "Validation")
+	bool bUnloadAssetsLoadedForValidation = false;
 	
+	/** 
+	 * If false, external objects (e.g. actors stored separately from maps) will not be loaded when 
+	 * validating their associated asset (e.g. the map) 
+	 */
+	UPROPERTY(EditAnywhere, BlueprintReadWrite, Category = "Validation")
+	bool bLoadExternalObjectsForValidation = true;
+
 	/** If true, captures log warnings and errors from loading assets for validation and reports them as validation results */
 	UPROPERTY(EditAnywhere, BlueprintReadWrite, Category = "Validation")
 	bool bCaptureAssetLoadLogs = true;
