Fix trying to load external actors on a non-partitioned world or on a validation task with no external objects.

diff --git a/Engine/Plugins/Editor/DataValidation/Source/DataValidation/Private/LoadExternalObjectsForValidation.cpp b/Engine/Plugins/Editor/DataValidation/Source/DataValidation/Private/LoadExternalObjectsForValidation.cpp
--- a/Engine/Plugins/Editor/DataValidation/Source/DataValidation/Private/LoadExternalObjectsForValidation.cpp
+++ b/Engine/Plugins/Editor/DataValidation/Source/DataValidation/Private/LoadExternalObjectsForValidation.cpp
@@ -17,12 +17,12 @@ namespace UE::DataValidation
 
 FScopedLoadExternalObjects::FScopedLoadExternalObjects(UObject* InAsset, FDataValidationContext& InContext, bool bEnabled)
 {
-    if (!bEnabled)
+    if (!bEnabled || InContext.GetAssociatedExternalObjects().Num() == 0)
     {
         return;
     }
 
-    if (UWorld* World = Cast<UWorld>(InAsset))
+    if (UWorld* World = Cast<UWorld>(InAsset); World && World->IsPartitionedWorld())
     {
         auto RegisterContainerToValidate = [this](UWorld* InWorld, FName InContainerPackageName, FActorDescContainerInstanceCollection& OutRegisteredContainers, const FGuid& InContentBundleGuid = FGuid(), const UExternalDataLayerAsset* InExternalDataLayerAsset = nullptr)
         {
