Added a method of disabling "validate on save" temporarily around places where it may behave incorrectly or perform poorly

diff --git a/Engine/Plugins/Editor/DataValidation/Source/DataValidation/Private/EditorValidatorSubsystem.cpp b/Engine/Plugins/Editor/DataValidation/Source/DataValidation/Private/EditorValidatorSubsystem.cpp
--- a/Engine/Plugins/Editor/DataValidation/Source/DataValidation/Private/EditorValidatorSubsystem.cpp
+++ b/Engine/Plugins/Editor/DataValidation/Source/DataValidation/Private/EditorValidatorSubsystem.cpp
@@ -52,6 +52,23 @@ UDataValidationSettings::UDataValidationSettings()
 
 }
 
+FScopedDisableValidateOnSave::FScopedDisableValidateOnSave()
+	: EditorValidationSubsystem(GEditor ? GEditor->GetEditorSubsystem<UEditorValidatorSubsystem>() : nullptr)
+{
+	if (EditorValidationSubsystem)
+	{
+		EditorValidationSubsystem->PushDisableValidateOnSave();
+	}
+}
+
+FScopedDisableValidateOnSave::~FScopedDisableValidateOnSave()
+{
+	if (EditorValidationSubsystem)
+	{
+		EditorValidationSubsystem->PopDisableValidateOnSave();
+	}
+}
+
 UEditorValidatorSubsystem::UEditorValidatorSubsystem()
 	: UEditorSubsystem()
 {
@@ -192,6 +209,36 @@ void UEditorValidatorSubsystem::Deinitialize()
 	Super::Deinitialize();
 }
 
+void UEditorValidatorSubsystem::PushDisableValidateOnSave()
+{
+	checkf(DisableValidateOnSaveCount < std::numeric_limits<decltype(DisableValidateOnSaveCount)>::max(), TEXT("PushDisableValidateOnSave overflow!"));
+	++DisableValidateOnSaveCount;
+}
+
+void UEditorValidatorSubsystem::PopDisableValidateOnSave()
+{
+	checkf(DisableValidateOnSaveCount > 0, TEXT("PopDisableValidateOnSave underflow!"));
+	--DisableValidateOnSaveCount;
+}
+
+bool UEditorValidatorSubsystem::ShouldValidateOnSave(bool bProceduralSave) const
+{
+	// Skip if not enabled
+	if (DisableValidateOnSaveCount > 0 || !GetDefault<UDataValidationSettings>()->bValidateOnSave)
+	{
+		return false;
+	}
+
+	// Skip auto and procedural saves
+	// For performance reasons, don't validate when making a procedural save by default. Assumption is we validated when saving previously. 
+	if (bProceduralSave || GEditor->IsAutosaving(EPackageAutoSaveType::Any))
+	{
+		return false;
+	}
+
+	return true;
+}
+
 void UEditorValidatorSubsystem::AddValidator(UEditorValidatorBase* InValidator)
 {
 	if (InValidator)
@@ -704,13 +751,7 @@ void UEditorValidatorSubsystem::LogAssetValidationSummary(FMessageLog& DataValid
 
 void UEditorValidatorSubsystem::ValidateOnSave(TArray<FAssetData> AssetDataList, bool bProceduralSave) const
 {
-	// Only validate if enabled and not auto saving
-	if (!GetDefault<UDataValidationSettings>()->bValidateOnSave || GEditor->IsAutosaving(EPackageAutoSaveType::Any))
-	{
-		return;
-	}
-
-	if (bProceduralSave)
+	if (!ShouldValidateOnSave(bProceduralSave))
 	{
 		return;
 	}
@@ -723,6 +764,11 @@ void UEditorValidatorSubsystem::ValidateOnSave(TArray<FAssetData> AssetDataList,
 		}));
 	}
 
+	if (AssetDataList.IsEmpty())
+	{
+		return;
+	}
+
 	FText SavedAsset = AssetDataList.Num() == 1 ? FText::FromName(AssetDataList[0].AssetName) : LOCTEXT("MultipleAssets", "multiple assets");
 	FValidateAssetsResults Results;
 
@@ -744,14 +790,7 @@ void UEditorValidatorSubsystem::ValidateOnSave(TArray<FAssetData> AssetDataList,
 
 void UEditorValidatorSubsystem::ValidateSavedPackage(FName PackageName, bool bProceduralSave)
 {
-	// Only validate if enabled and not auto saving
-	if (!GetDefault<UDataValidationSettings>()->bValidateOnSave || GEditor->IsAutosaving(EPackageAutoSaveType::Any))
-	{
-		return;
-	}
-
-	// For performance reasons, don't validate when making a procedural save by default. Assumption is we validated when saving previously. 
-	if (bProceduralSave)
+	if (!ShouldValidateOnSave(bProceduralSave))
 	{
 		return;
 	}
diff --git a/Engine/Plugins/Editor/DataValidation/Source/DataValidation/Public/EditorValidatorSubsystem.h b/Engine/Plugins/Editor/DataValidation/Source/DataValidation/Public/EditorValidatorSubsystem.h
--- a/Engine/Plugins/Editor/DataValidation/Source/DataValidation/Public/EditorValidatorSubsystem.h
+++ b/Engine/Plugins/Editor/DataValidation/Source/DataValidation/Public/EditorValidatorSubsystem.h
@@ -11,8 +11,9 @@
 #include "UObject/Package.h"
 #include "EditorValidatorSubsystem.generated.h"
 
-class UEditorValidatorBase;
 class UDataValidationChangelist;
+class UEditorValidatorBase;
+class UEditorValidatorSubsystem;
 struct FDirectoryPath;
 class FMessageLog;
 struct FARFilter;
@@ -183,6 +184,22 @@ struct DATAVALIDATION_API FValidateAssetsSettings
 	bool bSilent = false;
 };
 
+/**
+ * Utility to disable "validate on save" in UEditorValidatorSubsystem for the lifetime of this object.
+ * Calls UEditorValidatorSubsystem::PushDisableValidateOnSave on construction and UEditorValidatorSubsystem::PopDisableValidateOnSave on destruction.
+ */
+struct DATAVALIDATION_API FScopedDisableValidateOnSave
+{
+public:
+	FScopedDisableValidateOnSave();
+	~FScopedDisableValidateOnSave();
+
+	UE_NONCOPYABLE(FScopedDisableValidateOnSave);
+
+private:
+	UEditorValidatorSubsystem* EditorValidationSubsystem;
+};
+
 /**
  * UEditorValidatorSubsystem manages all the asset validation in the engine. 
  * 
@@ -211,6 +228,18 @@ public:
 	virtual void Initialize(FSubsystemCollectionBase& Collection);
 	virtual void Deinitialize();
 
+	/**
+	 * Push a new request to temporarily disable "validate on save".
+	 * Should be paired with a call to PopDisableValidateOnSave.
+	 */
+	void PushDisableValidateOnSave();
+
+	/**
+	 * Pop a previous request to temporarily disable "validate on save".
+	 * Should be paired with a call to PushDisableValidateOnSave.
+	 */
+	void PopDisableValidateOnSave();
+
 	/**
 	 * Called to validate assets from either the UI or a commandlet.
 	 * Loads the specified assets and runs all registered validators on them.
@@ -354,6 +383,8 @@ public:
 
 
 protected:
+	bool ShouldValidateOnSave(bool bProceduralSave) const;
+
 	void CleanupValidators();
 	
 	void WaitForAssetCompilationIfNecessary(EDataValidationUsecase InUsecase, bool bShowProgress = true) const;
@@ -428,4 +459,9 @@ protected:
 	bool bAllowBlueprintValidators;
 
 	bool bNeedLoadingOfValidators = false;
+
+	/**
+	 * Counter used by PushDisableValidateOnSave and PopDisableValidateOnSave to know whether "validate on save" is temporarily disabled.
+	 */
+	uint8 DisableValidateOnSaveCount = 0;
 };
