Change DataValidationTestActor so it can be properly loaded and unloaded in WP levels in the editor.

diff --git a/Engine/Plugins/Editor/DataValidation/Source/DataValidation/Private/DataValidationTestActor.cpp b/Engine/Plugins/Editor/DataValidation/Source/DataValidation/Private/DataValidationTestActor.cpp
--- a/Engine/Plugins/Editor/DataValidation/Source/DataValidation/Private/DataValidationTestActor.cpp
+++ b/Engine/Plugins/Editor/DataValidation/Source/DataValidation/Private/DataValidationTestActor.cpp
@@ -2,15 +2,49 @@
 
 #include "DataValidationTestActor.h"
 
+#include "AssetRegistry/AssetData.h"
+#include "Components/BillboardComponent.h"
+#include "Components/StaticMeshComponent.h"
+#include "Engine/Texture2D.h"
 #include "Misc/DataValidation.h"
+#include "UObject/ConstructorHelpers.h"
 
 #define LOCTEXT_NAMESPACE "AssetValidation"
 
+ADataValidationTestActor::ADataValidationTestActor(const FObjectInitializer& ObjectInitializer)
+	: Super(ObjectInitializer)
+{
+	SpriteComponent = CreateDefaultSubobject<UBillboardComponent>(TEXT("Sprite"));
+	RootComponent = SpriteComponent;
+
+	// Structure to hold one-time initialization
+	struct FConstructorStatics
+	{
+		ConstructorHelpers::FObjectFinderOptional<UTexture2D> SpriteTexture;
+		FName ID_Info;
+		FText NAME_Info;
+		FConstructorStatics()
+			: SpriteTexture(TEXT("/Engine/EditorResources/S_Actor"))
+			, ID_Info(TEXT("Info"))
+			, NAME_Info(NSLOCTEXT("SpriteCategory", "Info", "Info"))
+		{
+		}
+	};
+	static FConstructorStatics ConstructorStatics;
+
+	SpriteComponent->Sprite = ConstructorStatics.SpriteTexture.Get();
+	SpriteComponent->SpriteInfo.Category = ConstructorStatics.ID_Info;
+	SpriteComponent->SpriteInfo.DisplayName = ConstructorStatics.NAME_Info;
+	SpriteComponent->bIsScreenSizeScaled = true;
+
+	bIsSpatiallyLoaded = true;
+}
+
 EDataValidationResult ADataValidationTestActor::IsDataValid(FDataValidationContext& Context) const 
 {
     if (!bPassValidation)
     {
-        Context.AddError(LOCTEXT("bPassValidationFalse", "bPassValidation is false"));
+        Context.AddMessage(FAssetData(this), EMessageSeverity::Error, LOCTEXT("bPassValidationFalse", "bPassValidation is false"));
         return EDataValidationResult::Invalid;
     }
     return EDataValidationResult::Valid;
diff --git a/Engine/Plugins/Editor/DataValidation/Source/DataValidation/Private/DataValidationTestActor.h b/Engine/Plugins/Editor/DataValidation/Source/DataValidation/Private/DataValidationTestActor.h
--- a/Engine/Plugins/Editor/DataValidation/Source/DataValidation/Private/DataValidationTestActor.h
+++ b/Engine/Plugins/Editor/DataValidation/Source/DataValidation/Private/DataValidationTestActor.h
@@ -2,16 +2,20 @@
 
 #pragma once
 
-#include "GameFramework/Info.h"
+#include "GameFramework/Actor.h"
 
 #include "DataValidationTestActor.generated.h"
 
-UCLASS()
-class ADataValidationTestActor : public AInfo
+class UBillboardComponent;
+
+UCLASS(showcategories=(HLOD, WorldPartition, DataLayers, Transformation))
+class ADataValidationTestActor : public AActor
 {
     GENERATED_BODY()
 
 public:
+	ADataValidationTestActor(const FObjectInitializer& ObjectInitializer);
+
     UPROPERTY(EditAnywhere, BlueprintReadWrite, Category="Validation")
     bool bPassValidation = false;
 
@@ -19,4 +23,6 @@ public:
     EDataValidationResult IsDataValid(FDataValidationContext& Context) const override;
 #endif
 
+	UPROPERTY()
+	TObjectPtr<UBillboardComponent> SpriteComponent;
 };
\ No newline at end of file
