When validating actors in external packages, if the actor's world is streamed in to another world, validate that world instead.
When loading external objects for validation, respect the value of the bLoadAssetsForValidation setting.

diff --git a/Engine/Plugins/Editor/DataValidation/Source/DataValidation/Private/EditorValidatorSubsystem.cpp b/Engine/Plugins/Editor/DataValidation/Source/DataValidation/Private/EditorValidatorSubsystem.cpp
--- a/Engine/Plugins/Editor/DataValidation/Source/DataValidation/Private/EditorValidatorSubsystem.cpp
+++ b/Engine/Plugins/Editor/DataValidation/Source/DataValidation/Private/EditorValidatorSubsystem.cpp
@@ -438,6 +438,7 @@ EDataValidationResult UEditorValidatorSubsystem::ValidateAssetsInternal(
 {
 	IAssetRegistry& AssetRegistry = IAssetRegistry::GetChecked();
 
+	// The number of assets to validate may decrease from merging when dealing with external objects but it shouldn't increase
 	FScopedSlowTask SlowTask(AssetDataList.Num(), LOCTEXT("DataValidation.ValidateAssetsTask", "Validating Assets"));
 	if (!InSettings.bSilent)
 	{
@@ -463,16 +464,33 @@ EDataValidationResult UEditorValidatorSubsystem::ValidateAssetsInternal(
 	TMap<FAssetData, TArray<FAssetData>> AssetsToExternalObjects;
 	for (auto It = AssetDataList.CreateIterator(); It; ++It)
 	{
-		if (!It->GetOptionalOuterPathName().IsNone())
+		if (It->GetOptionalOuterPathName().IsNone())
 		{
-			FSoftObjectPath Path = It->ToSoftObjectPath();
-			FAssetData OuterAsset = AssetRegistry.GetAssetByObjectPath(Path.GetWithoutSubPath(), true);
-			if (OuterAsset.IsValid())
+			// Standalone asset, leave it in the list
+			continue;
+		}
+
+		FAssetData OuterAsset = AssetRegistry.GetAssetByObjectPath(It->ToSoftObjectPath().GetWithoutSubPath(), true);
+		if (!OuterAsset.IsValid())
+		{
+			// We can't validate this asset if we can't find the package to load it into
+			It.RemoveCurrent();
+			continue;
+		}
+
+		// Special case for level instances in world partition - if the outer asset we'd like to validate is loaded & streamed in to 
+		// another world, validate that world instead. 
+		if (UWorld* AssetWorld = Cast<UWorld>(OuterAsset.FastGetAsset(false)))
+		{
+			if (   AssetWorld->PersistentLevel 
+				&& AssetWorld->PersistentLevel->OwningWorld != nullptr 
+				&& AssetWorld->PersistentLevel->OwningWorld != AssetWorld)
 			{
-				AssetsToExternalObjects.FindOrAdd(OuterAsset).Add(*It);
+				OuterAsset = AssetRegistry.GetAssetByObjectPath(FSoftObjectPath(AssetWorld->PersistentLevel->OwningWorld), true);
 			}
-			It.RemoveCurrent();
 		}
+		AssetsToExternalObjects.FindOrAdd(OuterAsset).Add(*It);
+		It.RemoveCurrent();
 	}
 	
 	// Add any packages which contain those external objects to be validated
@@ -503,6 +521,8 @@ EDataValidationResult UEditorValidatorSubsystem::ValidateAssetsInternal(
 	// Loaded assets ought to have the standalone flag and so not be garbage collected, but keep weak pointers here for memory safety
 	TArray<TWeakObjectPtr<UPackage>> PackagesToUnload;	
 
+	const bool bLoadExternalObjects = InSettings.bLoadExternalObjectsForValidation && InSettings.bLoadAssetsForValidation;
+
 	// Now add to map or update as needed
 	for (const FAssetData& Data : AssetDataList)
 	{
@@ -596,7 +616,7 @@ EDataValidationResult UEditorValidatorSubsystem::ValidateAssetsInternal(
 
 		if (LoadedAsset)
 		{
-			UE::DataValidation::FScopedLoadExternalObjects ExternalObjectsLoader(LoadedAsset, ValidationContext, InSettings.bLoadExternalObjectsForValidation);
+			UE::DataValidation::FScopedLoadExternalObjects ExternalObjectsLoader(LoadedAsset, ValidationContext, bLoadExternalObjects);
 
 			UE::DataValidation::FScopedLogMessageGatherer LogGatherer(InSettings.bCaptureLogsDuringValidation);
 			AssetResult = IsObjectValidWithContext(LoadedAsset, ValidationContext);
diff --git a/Engine/Source/Runtime/CoreUObject/Public/Misc/DataValidation.h b/Engine/Source/Runtime/CoreUObject/Public/Misc/DataValidation.h
--- a/Engine/Source/Runtime/CoreUObject/Public/Misc/DataValidation.h
+++ b/Engine/Source/Runtime/CoreUObject/Public/Misc/DataValidation.h
@@ -76,6 +76,11 @@ public:
 		return ValidationUsecase;
 	}
 	
+	/**
+	 * Other objects associated with the asset being validated. 
+	 * e.g. if the asset being validated is a world, these may be actors within that world's PersistentLevel with external packages.
+	 * May also be actors in a streamed in sublevel which are stored in external packages.
+	 */
 	TConstArrayView<FAssetData> GetAssociatedExternalObjects() const
 	{
 		return AssociatedExternalObjects;
