[UDN - 17520866] Missing initialization of FLogMessageGathererImpl IgnoreCategories and IgnorePatterns from Default

diff --git a/Engine/Plugins/Editor/DataValidation/Source/DataValidation/Private/EditorValidatorHelpers.cpp b/Engine/Plugins/Editor/DataValidation/Source/DataValidation/Private/EditorValidatorHelpers.cpp
--- a/Engine/Plugins/Editor/DataValidation/Source/DataValidation/Private/EditorValidatorHelpers.cpp
+++ b/Engine/Plugins/Editor/DataValidation/Source/DataValidation/Private/EditorValidatorHelpers.cpp
@@ -97,6 +97,9 @@ struct FLogMessageGathererImpl : public FOutputDevice
             DefaultIgnorePatterns = DefaultIgnorePatternsDelegate.Execute();
         }
 
+		IgnoreCategories = DefaultIgnoreCategories;
+		IgnorePatterns = DefaultIgnorePatterns;
+
         GLog->AddOutputDevice(this);
     }
     
