Improved the consistency of data validation auto-registration

Previously data validators were only auto-registered by UEditorValidatorSubsystem::Initialize (for C++ validators) and either UEditorValidatorSubsystem::Initialize or IAssetRegistry::OnFilesLoaded (for BP validators; whichever ran later).
This meant that C++ modules loaded past this point (of which there are hundreds) and BPs created/mounted after the initial asset scan were never considered for data validation auto-registration.
The method of detecting BP data validators was also skipping any cooked BPs, as it was searching for UEditorUtilityBlueprint assets, rather than BP classes derived from UEditorValidatorBase.

This change updates the Validators map in UEditorValidatorSubsystem to always map from a class (including BPGC) to a validator (never a BP), which makes its behavior consistent for cooked and uncooked BPs.
It also adds hooks for when modules are loaded/unloaded and when BP assets are created/mounted and deleted/unmounted, and uses these hooks to update the set of auto-registered validators accordingly.

diff --git a/Engine/Plugins/Editor/DataValidation/Source/DataValidation/Private/EditorValidatorSubsystem.cpp b/Engine/Plugins/Editor/DataValidation/Source/DataValidation/Private/EditorValidatorSubsystem.cpp
--- a/Engine/Plugins/Editor/DataValidation/Source/DataValidation/Private/EditorValidatorSubsystem.cpp
+++ b/Engine/Plugins/Editor/DataValidation/Source/DataValidation/Private/EditorValidatorSubsystem.cpp
@@ -13,6 +13,7 @@
 #include "EditorValidatorHelpers.h"
 #include "DirectoryWatcherModule.h"
 #include "Editor.h"
+#include "EditorClassUtils.h"
 #include "EditorUtilityBlueprint.h"
 #include "EditorValidatorBase.h"
 #include "Engine/BlueprintGeneratedClass.h"
@@ -95,40 +96,21 @@ bool UEditorValidatorSubsystem::ShouldCreateSubsystem(UObject* Outer) const
 
 void UEditorValidatorSubsystem::Initialize(FSubsystemCollectionBase& Collection)
 {
-	const FAssetRegistryModule& AssetRegistryModule = FModuleManager::LoadModuleChecked<FAssetRegistryModule>(TEXT("AssetRegistry"));
-	if (!AssetRegistryModule.Get().IsLoadingAssets())
+	// C++ registration
+	FCoreDelegates::OnPostEngineInit.AddUObject(this, &UEditorValidatorSubsystem::RegisterNativeValidators);
+
+	// Blueprint registration
+	if (IAssetRegistry& AssetRegistry = IAssetRegistry::GetChecked();
+		AssetRegistry.IsLoadingAssets())
 	{
-		RegisterBlueprintValidators();
+		// We are still discovering assets, listen for the completion delegate before building the validator list
+		AssetRegistry.OnFilesLoaded().AddUObject(this, &UEditorValidatorSubsystem::RegisterBlueprintValidators);
 	}
 	else
 	{
-		// We are still discovering assets, listen for the completion delegate before building the graph
-		if (!AssetRegistryModule.Get().OnFilesLoaded().IsBoundToObject(this))
-		{
-			AssetRegistryModule.Get().OnFilesLoaded().AddUObject(this, &UEditorValidatorSubsystem::RegisterBlueprintValidators);
-		}
+		RegisterBlueprintValidators();
 	}
 
-	// C++ registration
-	TArray<UClass*> ValidatorClasses;
-	GetDerivedClasses(UEditorValidatorBase::StaticClass(), ValidatorClasses);
-	for (UClass* ValidatorClass : ValidatorClasses)
-	{
-		if (!ValidatorClass->HasAllClassFlags(CLASS_Abstract))
-		{
-			UPackage* const ClassPackage = ValidatorClass->GetOuterUPackage();
-			if (ClassPackage)
-			{
-				const FName ModuleName = FPackageName::GetShortFName(ClassPackage->GetFName());
-				if (FModuleManager::Get().IsModuleLoaded(ModuleName))
-				{
-					UEditorValidatorBase* Validator = NewObject<UEditorValidatorBase>(GetTransientPackage(), ValidatorClass);
-					AddValidator(Validator);
-				}
-				}
-			}
-		}
-
 	// Register to SCC pre-submit callback
 	ISourceControlModule::Get().RegisterPreSubmitDataValidation(FSourceControlPreSubmitDataValidationDelegate::CreateUObject(this, &UEditorValidatorSubsystem::ValidateChangelistPreSubmit));
 }
@@ -164,39 +146,68 @@ bool UEditorValidatorSubsystem::ShouldValidateAsset(
 	return true;
 }
 
-// Rename to BP validators
+void UEditorValidatorSubsystem::RegisterNativeValidators()
+{
+	ensureAlwaysMsgf(!bHasRegisteredNativeValidators, TEXT("Native validators have already been registered!"));
+	if (bHasRegisteredNativeValidators)
+	{
+		return;
+	}
+
+	TArray<UClass*> ValidatorClasses;
+	GetDerivedClasses(UEditorValidatorBase::StaticClass(), ValidatorClasses);
+
+	for (UClass* ValidatorClass : ValidatorClasses)
+	{
+		// GetDerivedClasses may include a mix of C++ and loaded BP classes
+		// Skip any non-C++ classes, as well as anything that has already been registered by this point
+		if (!ValidatorClass->HasAllClassFlags(CLASS_Abstract) && !Validators.Contains(ValidatorClass->GetClassPathName()) 
+			&& FPackageName::IsScriptPackage(FNameBuilder(ValidatorClass->GetPackage()->GetFName())))
+		{
+			UEditorValidatorBase* Validator = NewObject<UEditorValidatorBase>(GetTransientPackage(), ValidatorClass);
+			AddValidator(Validator);
+		}
+	}
+
+	// Watch for native modules being added/removed
+	FModuleManager::Get().OnModulesChanged().AddUObject(this, &UEditorValidatorSubsystem::OnNativeModulesChanged);
+
+	bHasRegisteredNativeValidators = true;
+}
+
 void UEditorValidatorSubsystem::RegisterBlueprintValidators()
 {
+	ensureAlwaysMsgf(!bHasRegisteredBlueprintValidators, TEXT("Blueprint validators have already been registered!"));
+	if (bHasRegisteredBlueprintValidators)
+	{
+		return;
+	}
+
 	if (bAllowBlueprintValidators)
 	{
-		// Locate all validators (include unloaded)
-		const FAssetRegistryModule& AssetRegistryModule = FModuleManager::LoadModuleChecked<FAssetRegistryModule>(TEXT("AssetRegistry"));
-		TArray<FAssetData> AllBPsAssetData;
-		AssetRegistryModule.Get().GetAssetsByClass(UEditorUtilityBlueprint::StaticClass()->GetClassPathName(), AllBPsAssetData, true);
+		IAssetRegistry& AssetRegistry = IAssetRegistry::GetChecked();
 
-		for (FAssetData& BPAssetData : AllBPsAssetData)
-		{
-			UClass* ParentClass = nullptr;
-			FString ParentClassName;
-			if (!BPAssetData.GetTagValue(FBlueprintTags::NativeParentClassPath, ParentClassName))
-			{
-				BPAssetData.GetTagValue(FBlueprintTags::ParentClassPath, ParentClassName);
-			}
+		// Locate all Blueprint-based validator classes (include unloaded)
+		TSet<FTopLevelAssetPath> ValidatorClasses;
+		AssetRegistry.GetDerivedClassNames({ UEditorValidatorBase::StaticClass()->GetClassPathName() }, {}, ValidatorClasses);
 
-			if (!ParentClassName.IsEmpty())
+		for (const FTopLevelAssetPath& ValidatorClass : ValidatorClasses)
+		{
+			// IAssetRegistry::GetDerivedClassNames may include a mix of C++ and BP classes
+			// Skip any C++ classes, as well as anything that has already been registered by this point
+			if (!Validators.Contains(ValidatorClass) && !FPackageName::IsScriptPackage(FNameBuilder(ValidatorClass.GetPackageName())))
 			{
-				UObject* Outer = nullptr;
-				ResolveName(Outer, ParentClassName, false, false);
-				ParentClass = FindObject<UClass>(Outer, *ParentClassName);
-				if (!ParentClass || !ParentClass->IsChildOf(UEditorValidatorBase::StaticClass()))
-				{
-					continue;
-				}
+				AddValidator(ValidatorClass);
 			}
-
-			AddValidator(BPAssetData);
 		}
+
+		// Watch for BPs being added/removed
+		// Recompilation is handled by the standard reinstancing logic, as Validators is a UPROPERTY
+		AssetRegistry.OnAssetsAdded().AddUObject(this, &UEditorValidatorSubsystem::OnAssetsAdded);
+		AssetRegistry.OnAssetsRemoved().AddUObject(this, &UEditorValidatorSubsystem::OnAssetsRemoved);
 	}
+
+	bHasRegisteredBlueprintValidators = true;
 }
 
 void UEditorValidatorSubsystem::Deinitialize()
@@ -243,27 +254,18 @@ void UEditorValidatorSubsystem::AddValidator(UEditorValidatorBase* InValidator)
 {
 	if (InValidator)
 	{
-		UClass* Class = InValidator->GetClass();
-		if (Cast<UBlueprintGeneratedClass>(Class))
-		{
-			if(UObject* ClassGenerator = Class->ClassGeneratedBy)
-			{
-				Validators.Add(FTopLevelAssetPath(ClassGenerator), InValidator);
-			}
-		}
-		else
-		{ 
-			Validators.Add(InValidator->GetClass()->GetClassPathName(), InValidator);
-		}
+		const FTopLevelAssetPath ValidatorClass = InValidator->GetClass()->GetClassPathName();
+		Validators.Add(ValidatorClass, InValidator);
+		ValidatorClassesPendingLoad.Remove(ValidatorClass);
 	}
 }
 
-void UEditorValidatorSubsystem::AddValidator(const FAssetData& InValidatorAssetData)
+void UEditorValidatorSubsystem::AddValidator(const FTopLevelAssetPath& InValidatorClass)
 {
-	if (InValidatorAssetData.IsValid())
+	if (InValidatorClass.IsValid())
 	{
-		Validators.Add(FTopLevelAssetPath(InValidatorAssetData.PackageName, InValidatorAssetData.AssetName), nullptr);
-		bNeedLoadingOfValidators = true;
+		Validators.Add(InValidatorClass, nullptr);
+		ValidatorClassesPendingLoad.Add(InValidatorClass);
 	}
 }
 
@@ -271,29 +273,30 @@ void UEditorValidatorSubsystem::RemoveValidator(UEditorValidatorBase* InValidato
 {
 	if (InValidator)
 	{
-		UClass* Class = InValidator->GetClass();
-		if (Cast<UBlueprintGeneratedClass>(Class))
-		{
-			if (UObject* ClassGenerator = Class->ClassGeneratedBy)
-			{
-				Validators.Remove(FTopLevelAssetPath(ClassGenerator));
-			}
-		}
-		else
-		{
-			Validators.Remove(InValidator->GetClass()->GetClassPathName());
-		}
+		RemoveValidator(InValidator->GetClass()->GetClassPathName());
+	}
+}
+
+void UEditorValidatorSubsystem::RemoveValidator(const FTopLevelAssetPath& InValidatorClass)
+{
+	if (InValidatorClass.IsValid())
+	{
+		Validators.Remove(InValidatorClass);
+		ValidatorClassesPendingLoad.Remove(InValidatorClass);
 	}
 }
 
 void UEditorValidatorSubsystem::CleanupValidators()
 {
 	Validators.Empty();
+	ValidatorClassesPendingLoad.Empty();
+	NativeModulesPendingLoad.Empty();
+	NativeModulesPendingUnload.Empty();
 }
 
 void UEditorValidatorSubsystem::ForEachEnabledValidator(TFunctionRef<bool(UEditorValidatorBase* Validator)> Callback) const
 {
-	LoadValidators();
+	UpdateValidators();
 
 	for (const TPair<FTopLevelAssetPath, TObjectPtr<UEditorValidatorBase>>& ValidatorPair : Validators)
 	{
@@ -1070,42 +1073,149 @@ void UEditorValidatorSubsystem::GatherAssetsToValidateFromChangelist(
 	}
 }
 
-void UEditorValidatorSubsystem::LoadValidators() const
+void UEditorValidatorSubsystem::OnNativeModulesChanged(FName InModuleName, EModuleChangeReason InModuleChangeReason)
+{
+	switch (InModuleChangeReason)
+	{
+		case EModuleChangeReason::ModuleLoaded:
+			NativeModulesPendingLoad.Add(InModuleName);
+			NativeModulesPendingUnload.Remove(InModuleName);
+			break;
+
+		case EModuleChangeReason::ModuleUnloaded:
+			NativeModulesPendingUnload.Add(InModuleName);
+			NativeModulesPendingLoad.Remove(InModuleName);
+			break;
+
+		default:
+			break;
+	}
+}
+
+void UEditorValidatorSubsystem::OnAssetsAdded(TConstArrayView<FAssetData> InAssets)
 {
-	const_cast<UEditorValidatorSubsystem*>(this)->LoadValidators();
+	OnAssetsAddedOrRemoved(InAssets, [this](const FTopLevelAssetPath& BPGC) { AddValidator(BPGC); });
 }
 
-void UEditorValidatorSubsystem::LoadValidators() 
+void UEditorValidatorSubsystem::OnAssetsRemoved(TConstArrayView<FAssetData> InAssets)
 {
-	if (bNeedLoadingOfValidators)
+	OnAssetsAddedOrRemoved(InAssets, [this](const FTopLevelAssetPath& BPGC) { RemoveValidator(BPGC); });
+}
+
+void UEditorValidatorSubsystem::OnAssetsAddedOrRemoved(TConstArrayView<FAssetData> InAssets, TFunctionRef<void(const FTopLevelAssetPath& BPGC)> Callback)
+{
+	for (const FAssetData& Asset : InAssets)
 	{
-		for (TPair<FTopLevelAssetPath, TObjectPtr<UEditorValidatorBase>>& ValidatorPair : Validators)
+		if (const UClass* AssetClass = Asset.GetClass())
 		{
-			if (!ValidatorPair.Value)
+			if (AssetClass->IsChildOf<UEditorUtilityBlueprint>())
+			{
+				// Uncooked BP
+				if (const UClass* ParentClass = UBlueprint::GetBlueprintParentClassFromAssetTags(Asset);
+					ParentClass && ParentClass->IsChildOf<UEditorValidatorBase>())
+				{
+					FTopLevelAssetPath BPGC = FEditorClassUtils::GetClassPathNameFromAssetTag(Asset);
+					if (BPGC.IsNull())
+					{
+						BPGC = FTopLevelAssetPath(Asset.PackageName, FName(WriteToString<128>(Asset.AssetName, TEXT("_C"))));
+					}
+					Callback(BPGC);
+				}
+			}
+			else if (AssetClass->IsChildOf<UBlueprintGeneratedClass>())
 			{
-				FSoftObjectPath ValidatorObjectSoftPath(ValidatorPair.Key);
-				UObject* ValidatorObject = ValidatorObjectSoftPath.ResolveObject();
+				// Cooked BPGC
+				if (const UClass* ParentClass = UBlueprint::GetBlueprintParentClassFromAssetTags(Asset);
+					ParentClass && ParentClass->IsChildOf<UEditorValidatorBase>())
+				{
+					Callback(FTopLevelAssetPath(Asset.PackageName, Asset.AssetName));
+				}
+			}
+		}
+	}
+}
 
-				// If this object isn't currently loaded, load it
-				if (ValidatorObject == nullptr)
+void UEditorValidatorSubsystem::UpdateValidators() const
+{
+	const_cast<UEditorValidatorSubsystem*>(this)->UpdateValidators();
+}
+
+void UEditorValidatorSubsystem::UpdateValidators()
+{
+	UE_CLOG(!bHasRegisteredNativeValidators, LogContentValidation, Warning, TEXT("UpdateValidators request made before RegisterNativeValidators. Native validators may be missing!"));
+	UE_CLOG(!bHasRegisteredBlueprintValidators, LogContentValidation, Warning, TEXT("UpdateValidators request made before RegisterBlueprintValidators. Blueprint validators may be missing!"));
+
+	// Remove any existing validators for unloaded modules
+	if (NativeModulesPendingUnload.Num() > 0)
+	{
+		TSet<FName> ModulePackageNames;
+		ModulePackageNames.Reserve(NativeModulesPendingUnload.Num());
+		Algo::Transform(NativeModulesPendingUnload, ModulePackageNames, [](const FName ModuleName) { return FPackageName::GetModuleScriptPackageName(ModuleName); });
+
+		for (auto It = Validators.CreateIterator(); It; ++It)
+		{
+			if (ModulePackageNames.Contains(It->Key.GetPackageName()))
+			{
+				It.RemoveCurrent();
+			}
+		}
+
+		NativeModulesPendingUnload.Reset();
+	}
+
+	// Add any new validators for loaded modules
+	if (NativeModulesPendingLoad.Num() > 0)
+	{
+		TSet<FName> ModulePackageNames;
+		ModulePackageNames.Reserve(NativeModulesPendingLoad.Num());
+		Algo::Transform(NativeModulesPendingLoad, ModulePackageNames, [](const FName ModuleName) { return FPackageName::GetModuleScriptPackageName(ModuleName); });
+
+		// GetDerivedClasses has an accelerator table, so it's faster to query all the native UEditorValidatorBase classes (of which there will be relatively few) and 
+		// then filter them down by module, than it is to get all the classes of each pending module (of which there may be many) and then filter them by type
+		TArray<UClass*> ValidatorClasses;
+		GetDerivedClasses(UEditorValidatorBase::StaticClass(), ValidatorClasses);
+
+		for (UClass* ValidatorClass : ValidatorClasses)
+		{
+			// GetDerivedClasses may include a mix of C++ and loaded BP classes
+			// Skip any classes outside of the modules requested, as well as anything that has already been registered by this point
+			if (!ValidatorClass->HasAllClassFlags(CLASS_Abstract) && !Validators.Contains(ValidatorClass->GetClassPathName())
+				&& ModulePackageNames.Contains(ValidatorClass->GetPackage()->GetFName()))
+			{
+				UEditorValidatorBase* Validator = NewObject<UEditorValidatorBase>(GetTransientPackage(), ValidatorClass);
+				AddValidator(Validator);
+			}
+		}
+
+		NativeModulesPendingLoad.Reset();
+	}
+
+	// Add any new validators from pending Blueprint classes
+	if (ValidatorClassesPendingLoad.Num() > 0)
+	{
+		for (const FTopLevelAssetPath& ValidatorClassPendingLoad : ValidatorClassesPendingLoad)
+		{
+			if (TObjectPtr<UEditorValidatorBase>* ValidatorInstance = Validators.Find(ValidatorClassPendingLoad);
+				ensure(ValidatorInstance))
+			{
+				TSoftClassPtr<UEditorValidatorBase> ValidatorClassSoftPtr = TSoftClassPtr<UEditorValidatorBase>(FSoftObjectPath(ValidatorClassPendingLoad));
+				UClass* ValidatorClass = ValidatorClassSoftPtr.Get();
+
+				// If this class isn't currently loaded, load it
+				if (!ValidatorClass)
 				{
 					FCookLoadScope EditorOnlyLoadScope(ECookLoadType::EditorOnly);
 					FSoftObjectPathSerializationScope SerializationScope(NAME_None, NAME_None, ESoftObjectPathCollectType::EditorOnlyCollect, ESoftObjectPathSerializeType::AlwaysSerialize);
-					ValidatorObject = ValidatorObjectSoftPath.TryLoad();
+					ValidatorClass = ValidatorClassSoftPtr.LoadSynchronous();
 				}
 
-				if (ValidatorObject)
+				if (ValidatorClass && !ValidatorClass->HasAnyClassFlags(CLASS_Abstract))
 				{
-					if (UEditorUtilityBlueprint* ValidatorBlueprint = Cast<UEditorUtilityBlueprint>(ValidatorObject))
-					{
-						UEditorValidatorBase* Validator = NewObject<UEditorValidatorBase>(GetTransientPackage(), ValidatorBlueprint->GeneratedClass);
-						ValidatorPair.Value = Validator;
-					}
+					*ValidatorInstance = NewObject<UEditorValidatorBase>(GetTransientPackage(), ValidatorClass);
 				}
 			}
 		}
-
-		bNeedLoadingOfValidators = false;
+		ValidatorClassesPendingLoad.Reset();
 	}
 }
 
diff --git a/Engine/Plugins/Editor/DataValidation/Source/DataValidation/Public/EditorValidatorSubsystem.h b/Engine/Plugins/Editor/DataValidation/Source/DataValidation/Public/EditorValidatorSubsystem.h
--- a/Engine/Plugins/Editor/DataValidation/Source/DataValidation/Public/EditorValidatorSubsystem.h
+++ b/Engine/Plugins/Editor/DataValidation/Source/DataValidation/Public/EditorValidatorSubsystem.h
@@ -21,6 +21,8 @@ struct FAssetData;
 class ISourceControlChangelist;
 typedef TSharedPtr<class ISourceControlChangelist, ESPMode::ThreadSafe> FSourceControlChangelistPtr;
 
+enum class EModuleChangeReason;
+
 namespace EMessageSeverity
 {
 	enum Type : int;
@@ -399,6 +401,8 @@ protected:
 	 */
 	void ValidateAllSavedPackages();
 
+	void RegisterNativeValidators();
+
 	void RegisterBlueprintValidators();
 
 	/**
@@ -406,13 +410,33 @@ protected:
 	 */
 	void ValidateChangelistPreSubmit(FSourceControlChangelistPtr Changelist, EDataValidationResult& OutResult, TArray<FText>& ValidationErrors, TArray<FText>& ValidationWarnings) const;
 
-	void LoadValidators() const;
-	void LoadValidators();
+	/**
+	 * Handle native modules being loaded or unloaded (@see NativeModulesPendingLoad and NativeModulesPendingUnload)
+	 */
+	void OnNativeModulesChanged(FName InModuleName, EModuleChangeReason InModuleChangeReason);
+
+	/**
+	 * Handle Blueprint assets being added or removed (@see ValidatorClassesPendingLoad)
+	 */
+	void OnAssetsAdded(TConstArrayView<FAssetData> InAssets);
+	void OnAssetsRemoved(TConstArrayView<FAssetData> InAssets);
+	void OnAssetsAddedOrRemoved(TConstArrayView<FAssetData> InAssets, TFunctionRef<void(const FTopLevelAssetPath& BPGC)> Callback);
+
+	/**
+	 * Apply any pending changes to the list of active validators (@see ValidatorClassesPendingLoad, NativeModulesPendingLoad, and NativeModulesPendingUnload)
+	 */
+	void UpdateValidators() const;
+	void UpdateValidators();
 
 	/*
-	* Adds a validator to the list making sure it is a unique instance, but wait for the first use to load the assets.
-	*/
-	void AddValidator(const FAssetData&	InValidatorAssetData);
+	 * Adds a validator to the list making sure it is a unique instance, but wait for the first use to load the class (for Blueprints).
+	 */
+	void AddValidator(const FTopLevelAssetPath& InValidatorClass);
+
+	/*
+	 * Remove any validators (active or pending) for the given class.
+	 */
+	void RemoveValidator(const FTopLevelAssetPath& InValidatorClass);
 
 	/** 
 	 * Validate a set of assets, adding the results to the log/output of a higher level task (e.g. changelist validation).
@@ -451,15 +475,35 @@ protected:
 	/** List of saved package names to validate next frame */
 	TArray<FName> SavedPackagesToValidate;
 
+	/**
+	 * Active validator instances, mapped from their class.
+	 * @note Some instances may be null until UpdateValidators has been called (@see ValidatorClassesPendingLoad)
+	 */
 	UPROPERTY(Transient)
 	TMap<FTopLevelAssetPath, TObjectPtr<UEditorValidatorBase>> Validators;
 
+	/**
+	 * Set of Blueprint validator classes (from Validators) that have been discovered since the set of validators was last updated, and need to be loaded by the next call to UpdateValidators.
+	 */
+	TSet<FTopLevelAssetPath> ValidatorClassesPendingLoad;
+
+	/**
+	 * Set of native modules that have been loaded since the set of validators was last updated, and should be queried for new validator classes by the next call to UpdateValidators.
+	 */
+	TSet<FName> NativeModulesPendingLoad;
+
+	/**
+	 * Set of native modules that have been unloaded since the set of validators was last updated, and should be removed from Validators by the next call to UpdateValidators.
+	 */
+	TSet<FName> NativeModulesPendingUnload;
+
+	bool bHasRegisteredNativeValidators = false;
+	bool bHasRegisteredBlueprintValidators = false;
+
 	/** Specifies whether or not to allow Blueprint validators */
 	UPROPERTY(config)
 	bool bAllowBlueprintValidators;
 
-	bool bNeedLoadingOfValidators = false;
-
 	/**
 	 * Counter used by PushDisableValidateOnSave and PopDisableValidateOnSave to know whether "validate on save" is temporarily disabled.
 	 */
