Allow overriding the MessageLog name used for validation errors

diff --git a/Engine/Plugins/Editor/DataValidation/Source/DataValidation/Private/EditorValidatorSubsystem.cpp b/Engine/Plugins/Editor/DataValidation/Source/DataValidation/Private/EditorValidatorSubsystem.cpp
--- a/Engine/Plugins/Editor/DataValidation/Source/DataValidation/Private/EditorValidatorSubsystem.cpp
+++ b/Engine/Plugins/Editor/DataValidation/Source/DataValidation/Private/EditorValidatorSubsystem.cpp
@@ -40,6 +40,7 @@ DEFINE_LOG_CATEGORY(LogContentValidation);
 
 FValidateAssetsSettings::FValidateAssetsSettings()
 	: ShowMessageLogSeverity(EMessageSeverity::Warning)
+	, MessageLogName(UE::DataValidation::MessageLogName)
 	, MessageLogPageTitle(LOCTEXT("DataValidation.MessagePageTitle", "Data Validation"))
 {
 }
@@ -421,7 +422,7 @@ int32 UEditorValidatorSubsystem::ValidateAssetsWithSettings(
 	const FValidateAssetsSettings& InSettings,
 	FValidateAssetsResults& OutResults) const
 {
-	FMessageLog DataValidationLog(UE::DataValidation::MessageLogName);
+	FMessageLog DataValidationLog(InSettings.MessageLogName);
 	DataValidationLog.SetCurrentPage(InSettings.MessageLogPageTitle);
 	ValidateAssetsInternal(DataValidationLog, TSet<FAssetData>{AssetDataList}, InSettings, OutResults);
 	
@@ -803,7 +804,7 @@ void UEditorValidatorSubsystem::ValidateOnSave(TArray<FAssetData> AssetDataList,
 
 	if (ValidateAssetsWithSettings(AssetDataList, Settings, Results) > 0)
 	{
-		FMessageLog DataValidationLog(UE::DataValidation::MessageLogName);
+		FMessageLog DataValidationLog(Settings.MessageLogName);
 
 		const FText ErrorMessageNotification = FText::Format(
 			LOCTEXT("ValidationFailureNotification", "Validation failed when saving {0}, check Data Validation log"), SavedAsset);
@@ -946,7 +947,7 @@ EDataValidationResult UEditorValidatorSubsystem::ValidateChangelistsInternal(
 		return EDataValidationResult::NotValidated;
 	}
 
-	FMessageLog DataValidationLog(UE::DataValidation::MessageLogName);
+	FMessageLog DataValidationLog(Settings.MessageLogName);
 
 	// Choose a specific message log page for this output, flushing in case other recursive calls also write to this log
 	DataValidationLog.SetCurrentPage(Settings.MessageLogPageTitle);
diff --git a/Engine/Plugins/Editor/DataValidation/Source/DataValidation/Public/EditorValidatorSubsystem.h b/Engine/Plugins/Editor/DataValidation/Source/DataValidation/Public/EditorValidatorSubsystem.h
--- a/Engine/Plugins/Editor/DataValidation/Source/DataValidation/Public/EditorValidatorSubsystem.h
+++ b/Engine/Plugins/Editor/DataValidation/Source/DataValidation/Public/EditorValidatorSubsystem.h
@@ -179,6 +179,9 @@ struct DATAVALIDATION_API FValidateAssetsSettings
 	 */
 	TOptional<EMessageSeverity::Type> ShowMessageLogSeverity;
 
+	/** The name of the message log to use for warnings/errors/etc */
+	FName MessageLogName;
+
 	/** Title of message log page to use for warnings/errors/etc */
 	FText MessageLogPageTitle;
 
