Changelist Data Validation: output referencer for missing references when validating changelists.

diff --git a/Engine/Plugins/Editor/DataValidation/Source/DataValidation/Private/DataValidationChangelist.cpp b/Engine/Plugins/Editor/DataValidation/Source/DataValidation/Private/DataValidationChangelist.cpp
--- a/Engine/Plugins/Editor/DataValidation/Source/DataValidation/Private/DataValidationChangelist.cpp
+++ b/Engine/Plugins/Editor/DataValidation/Source/DataValidation/Private/DataValidationChangelist.cpp
@@ -16,9 +16,9 @@
 
 #define LOCTEXT_NAMESPACE "DataValidationChangelist"
 
-void UDataValidationChangelist::GatherDependencies(const FName& InPackageName, TSet<FName>& OutDependencies)
+void UDataValidationChangelist::GatherDependencies(const FName& InPackageName, TSet<TPair<FName, FName>>& OutDependencies)
 {
-	OutDependencies.Add(InPackageName);
+	OutDependencies.Add({ InPackageName, InPackageName });
 
 	TArray<FAssetData> Assets;
 	TArray<FName> Dependencies;
@@ -29,7 +29,7 @@ void UDataValidationChangelist::GatherDependencies(const FName& InPackageName, T
 		// Exclude script/memory packages
 		if (FPackageName::IsValidLongPackageName(PackageDependency.ToString()))
 		{
-			OutDependencies.Add(PackageDependency);
+			OutDependencies.Add({ InPackageName, PackageDependency });
 		}
 	}
 }
@@ -86,7 +86,7 @@ EDataValidationResult UDataValidationChangelist::IsDataValid(FDataValidationCont
 
 	// Gather dependencies of every file in the changelist
 	TArray<FName> FilesInChangelist;
-	TSet<FName> ExternalDependenciesSet;
+	TSet<TPair<FName, FName>> ExternalDependenciesSet;
 	
 	for (FName PackageName : ModifiedPackageNames)
 	{
@@ -96,11 +96,11 @@ EDataValidationResult UDataValidationChangelist::IsDataValid(FDataValidationCont
 
 	// For every dependency in the external dependencies that is not in the changelist
 	// Make sure that the source control state is "not currently modified"
-	TArray<FName> ExternalDependencies = ExternalDependenciesSet.Array();
+	TArray<TPair<FName, FName>> ExternalDependencies = ExternalDependenciesSet.Array();
 
-	ExternalDependencies.RemoveAll([&FilesInChangelist](FName& File) -> bool {
-		return Algo::AnyOf(FilesInChangelist, [&File](const auto& FileInChangelist) {
-			return File == FileInChangelist;
+	ExternalDependencies.RemoveAll([&FilesInChangelist](TPair<FName, FName>& FilePair) -> bool {
+		return Algo::AnyOf(FilesInChangelist, [&FilePair](const auto& FileInChangelist) {
+			return FilePair.Value == FileInChangelist;
 			});
 		});
 
@@ -108,9 +108,9 @@ EDataValidationResult UDataValidationChangelist::IsDataValid(FDataValidationCont
 
 	TArray<FString> ExternalDependenciesFilenames;
 
-	Algo::Transform(ExternalDependencies, ExternalDependenciesFilenames, [](const FName& InFilename) -> FString
+	Algo::Transform(ExternalDependencies, ExternalDependenciesFilenames, [](const TPair<FName, FName>& InFilenamePair) -> FString
 	{
-		return USourceControlHelpers::PackageFilename(InFilename.ToString());
+		return USourceControlHelpers::PackageFilename(InFilenamePair.Value.ToString());
 	});
 
 	// Update External dependencies state in case it changed from what is in cache
@@ -124,7 +124,7 @@ EDataValidationResult UDataValidationChangelist::IsDataValid(FDataValidationCont
 	for (int32 i = 0; i < ExternalDependenciesFilenames.Num(); ++i)
 	{
 		const FString& ExternalPackageFilename = ExternalDependenciesFilenames[i];
-		const FName& ExternalDependency = ExternalDependencies[i];
+		const TPair<FName, FName>& ExternalDependencyPair = ExternalDependencies[i];
 
 		FSourceControlStatePtr ExternalDependencyFileState = SourceControlProvider.GetState(ExternalPackageFilename, EStateCacheUsage::Use);
 
@@ -136,16 +136,20 @@ EDataValidationResult UDataValidationChangelist::IsDataValid(FDataValidationCont
 
 		// Some kinds of package dependencies create false positives because they're not checked into source
 		// control as of this writing (1/23/2023). We just skip over them.
-		if (FPackageName::IsVersePackage(ExternalDependency.ToString()))
+		if (FPackageName::IsVersePackage(ExternalDependencyPair.Value.ToString()))
 		{
 			continue;
 		}
 
+		auto GetReferencerText = [&ExternalDependencyPair]() { return FText::FromString(GetPrettyPackageName(ExternalDependencyPair.Key)); };
+		auto GetReferenceText = [&ExternalDependencyPair]() { return FText::FromString(GetPrettyPackageName(ExternalDependencyPair.Value)); };
+		auto GetExternalPackageFilenameText = [&ExternalPackageFilename]() { return FText::FromString(ExternalPackageFilename); };
+
 		// Dependency is checked out or added but is not in this changelist
 		if (ExternalDependencyFileState->IsCheckedOut() || ExternalDependencyFileState->IsAdded())
 		{
 			bHasChangelistErrors = true;
-			FText CurrentError = FText::Format(LOCTEXT("DataValidation.Changelist.Error", "{0} is missing from this changelist."), FText::FromString(GetPrettyPackageName(ExternalDependency)));
+			FText CurrentError = FText::Format(LOCTEXT("DataValidation.Changelist.Error", "{0} is missing from this changelist (referenced from {1})."), GetReferenceText(), GetReferencerText());
 			Context.AddError(CurrentError);
 		}
 		// Dependency is not at the latest revision
@@ -153,7 +157,7 @@ EDataValidationResult UDataValidationChangelist::IsDataValid(FDataValidationCont
 		{
 			if (!bIgnoreOutOfDateDependencies)
 			{
-				FText CurrentWarning = FText::Format(LOCTEXT("DataValidation.Changelist.NotLatest", "{0} is referenced but is not at the latest revision '{1}'"), FText::FromString(GetPrettyPackageName(ExternalDependency)), FText::FromString(ExternalPackageFilename));
+				FText CurrentWarning = FText::Format(LOCTEXT("DataValidation.Changelist.NotLatest", "{0} is referenced but is not at the latest revision '{1}' (referenced from {2})"), GetReferenceText(), GetExternalPackageFilenameText(), GetReferencerText());
 				Context.AddWarning(CurrentWarning);
 			}
 		}
@@ -163,13 +167,13 @@ EDataValidationResult UDataValidationChangelist::IsDataValid(FDataValidationCont
 			if (!FPaths::FileExists(ExternalDependencyFileState->GetFilename()))
 			{
 				bHasChangelistErrors = true;
-				FText CurrentError = FText::Format(LOCTEXT("DataValidation.Changelist.NotInWorkspace", "{0} is referenced and cannot be found in workspace '{1}'"), FText::FromString(GetPrettyPackageName(ExternalDependency)), FText::FromString(ExternalPackageFilename));
+				FText CurrentError = FText::Format(LOCTEXT("DataValidation.Changelist.NotInWorkspace", "{0} is referenced and cannot be found in workspace '{1}' (referenced from {2})"), GetReferenceText(), GetExternalPackageFilenameText(), GetReferencerText());
 				Context.AddError(CurrentError);
 			}
 			else
 			{
 				bHasChangelistErrors = true;
-				FText CurrentError = FText::Format(LOCTEXT("DataValidation.Changelist.NotInDepot", "{0} is referenced and must also be added to revision control '{1}'"), FText::FromString(GetPrettyPackageName(ExternalDependency)), FText::FromString(ExternalPackageFilename));
+				FText CurrentError = FText::Format(LOCTEXT("DataValidation.Changelist.NotInDepot", "{0} is referenced and must also be added to revision control '{1}' (referenced from {2})"), GetReferenceText(), GetExternalPackageFilenameText(), GetReferencerText());
 				Context.AddError(CurrentError);
 			}
 		}
diff --git a/Engine/Plugins/Editor/DataValidation/Source/DataValidation/Public/DataValidationChangelist.h b/Engine/Plugins/Editor/DataValidation/Source/DataValidation/Public/DataValidationChangelist.h
--- a/Engine/Plugins/Editor/DataValidation/Source/DataValidation/Public/DataValidationChangelist.h
+++ b/Engine/Plugins/Editor/DataValidation/Source/DataValidation/Public/DataValidationChangelist.h
@@ -30,7 +30,7 @@ public:
 	/** Initializes from an already-queried changelist state */
 	void Initialize(FSourceControlChangelistStateRef InChangelistState);
 
-	static void GatherDependencies(const FName& InPackageName, TSet<FName>& OutDependencies);
+	static void GatherDependencies(const FName& InPackageName, TSet<TPair<FName, FName>>& OutDependencies);
 	static FString GetPrettyPackageName(const FName& InPackageName);
 
 	/** Changelist to validate - may be null if this was constructed from a list of files */
