Explicitely UE_LOG shader compilation errors during material asset validation

diff --git a/Engine/Plugins/Editor/DataValidation/Source/DataValidation/Private/EditorValidatorSubsystem.cpp b/Engine/Plugins/Editor/DataValidation/Source/DataValidation/Private/EditorValidatorSubsystem.cpp
--- a/Engine/Plugins/Editor/DataValidation/Source/DataValidation/Private/EditorValidatorSubsystem.cpp
+++ b/Engine/Plugins/Editor/DataValidation/Source/DataValidation/Private/EditorValidatorSubsystem.cpp
@@ -693,6 +693,10 @@ EDataValidationResult UEditorValidatorSubsystem::ValidateAssetsInternal(
 							->AddToken(FTextToken::Create(LOCTEXT("DataValidation.ValidResult", "contains valid data.")));
 					}
 					break;
+				case EDataValidationResult::Invalid:
+					DataValidationLog.Info()->AddToken(FAssetDataToken::Create(Data))
+						->AddToken(FTextToken::Create(LOCTEXT("DataValidation.InvalidResult", "contains invalid data.")));
+					break;
 				case EDataValidationResult::NotValidated:
 					DataValidationLog.Info()->AddToken(FAssetDataToken::Create(Data))
 						->AddToken(FTextToken::Create(LOCTEXT("DataValidation.NotValidatedDataResult", "has no data validation.")));
diff --git a/Engine/Plugins/Editor/DataValidation/Source/DataValidation/Private/EditorValidator_Material.cpp b/Engine/Plugins/Editor/DataValidation/Source/DataValidation/Private/EditorValidator_Material.cpp
--- a/Engine/Plugins/Editor/DataValidation/Source/DataValidation/Private/EditorValidator_Material.cpp
+++ b/Engine/Plugins/Editor/DataValidation/Source/DataValidation/Private/EditorValidator_Material.cpp
@@ -19,6 +19,7 @@ UEditorValidator_Material::UEditorValidator_Material()
 		for (const FMaterialEditorValidationPlatform& Config: GetDefault<UDataValidationSettings>()->MaterialValidationPlatforms)
 		{
 			FShaderValidationPlatform Platform = {};
+			Platform.ShaderPlatformName = Config.ShaderPlatform.Name; 
 
 			bool bValidShaderPlatform = false;
 			if (Config.ShaderPlatform.Name == FMaterialEditorValidationShaderPlatform::MaxRHIShaderPlatformName)
@@ -74,17 +75,20 @@ bool UEditorValidator_Material::CanValidateAsset_Implementation(const FAssetData
 {
 	if (InContext.GetValidationUsecase() == EDataValidationUsecase::Save || InContext.GetValidationUsecase() == EDataValidationUsecase::Commandlet)
 	{
+		UE_LOG(LogContentValidation, Display, TEXT("skipping validation due to save or commandlet"));
 		return false;
 	}
 
 	if (ValidationPlatforms.IsEmpty())
 	{
+		UE_LOG(LogContentValidation, Display, TEXT("skipping validation due to empty validation platforms"));
 		return false;
 	}
 
 	const UMaterial* OriginalMaterial = Cast<UMaterial>(InAsset);
 	if (OriginalMaterial)
 	{
+		UE_LOG(LogContentValidation, Display, TEXT("enabling validation due to UMaterial"));
 		// We want to validate every UMaterial.
 		return true;
 	}
@@ -100,6 +104,7 @@ bool UEditorValidator_Material::CanValidateAsset_Implementation(const FAssetData
 			if (MaterialInstance->HasStaticParameters())
 			{
 				// We want to validate UMaterialInstance since it has a static parameter that could influence generated shader code.
+				UE_LOG(LogContentValidation, Display, TEXT("enabling validation due static parameters"));
 				return true;
 			}
 		}
@@ -119,10 +124,12 @@ EDataValidationResult UEditorValidator_Material::ValidateLoadedAsset_Implementat
 	UMaterial* Material = MaterialInstance ? MaterialInstance->GetMaterial() : DuplicateMaterial(OriginalMaterial);
 	if (!ensureAlways(OriginalMaterial) || !ensureAlways(Material))
 	{
+		UE_LOG(LogContentValidation, Display, TEXT("not validating because 1"));
 		return EDataValidationResult::NotValidated;
 	}
 
 	TArray<FMaterialResource*> Resources;
+	TMap<FMaterialResource*, TArray<FName>> ResourceToShaderPlatformNames;
 
 	for (const FShaderValidationPlatform& ValidationPlatform: ValidationPlatforms)
 	{
@@ -130,12 +137,15 @@ EDataValidationResult UEditorValidator_Material::ValidateLoadedAsset_Implementat
 
 		if (ensure(CurrentResource))
 		{
+			UE_LOG(LogContentValidation, Display, TEXT("caching shaders for '%s'"), *ValidationPlatform.ShaderPlatformName.ToString());
 			CurrentResource->CacheShaders(ValidationPlatform.ShaderPlatform);
+			ResourceToShaderPlatformNames.FindOrAdd(CurrentResource).Add(ValidationPlatform.ShaderPlatformName);
 		}
 	}
 
 	if (Resources.IsEmpty())
 	{
+		UE_LOG(LogContentValidation, Display, TEXT("not validating because 2"));
 		return EDataValidationResult::NotValidated;
 	}
 
@@ -149,11 +159,31 @@ EDataValidationResult UEditorValidator_Material::ValidateLoadedAsset_Implementat
 			UE_LOG(LogContentValidation, Warning, TEXT("Shader compilation was expected to be finished, but was not finished."));
 		}
 
-		bCompileErrors |= Resource->GetCompileErrors().Num() > 0;
+		TArray<FName>* ShaderPlatformNames = ResourceToShaderPlatformNames.Find(Resource);
+		const FString ShaderPlatformNameString = ShaderPlatformNames ? FString::JoinBy(*ShaderPlatformNames, TEXT(", "), [](const FName& Name) -> FString {return Name.ToString();}) : TEXT("Unknown");
+
+		for (const FString& ErrorText: Resource->GetCompileErrors())
+		{
+			UE_LOG(LogContentValidation, Error, TEXT("Failed to compile Material for platform %s due to '%s'."), *ShaderPlatformNameString, *ErrorText);
+			bCompileErrors = true;
+		}
+
+		UE_LOG(LogContentValidation, Display, TEXT("done with %s"), *ShaderPlatformNameString);
 	}
 
+	ResourceToShaderPlatformNames.Empty();
+
 	FMaterial::DeferredDeleteArray(Resources);
 
+	if (bCompileErrors)
+	{
+		UE_LOG(LogContentValidation, Display, TEXT("result - invalid"));
+	}
+	else
+	{
+		UE_LOG(LogContentValidation, Display, TEXT("result - ok"));
+	}
+
 	return !bCompileErrors ? EDataValidationResult::Valid : EDataValidationResult::Invalid;
 }
 
diff --git a/Engine/Plugins/Editor/DataValidation/Source/DataValidation/Public/EditorValidator_Material.h b/Engine/Plugins/Editor/DataValidation/Source/DataValidation/Public/EditorValidator_Material.h
--- a/Engine/Plugins/Editor/DataValidation/Source/DataValidation/Public/EditorValidator_Material.h
+++ b/Engine/Plugins/Editor/DataValidation/Source/DataValidation/Public/EditorValidator_Material.h
@@ -30,6 +30,7 @@ protected:
 private:
 	struct FShaderValidationPlatform
 	{
+		FName ShaderPlatformName; 
 		EShaderPlatform ShaderPlatform;
 		ERHIFeatureLevel::Type FeatureLevel;
 		EMaterialQualityLevel::Type MaterialQualityLevel;
