Add a cvar to explicitely enable asset validation in all contexts

diff --git a/Engine/Plugins/Editor/DataValidation/Source/DataValidation/Private/EditorValidator_Material.cpp b/Engine/Plugins/Editor/DataValidation/Source/DataValidation/Private/EditorValidator_Material.cpp
--- a/Engine/Plugins/Editor/DataValidation/Source/DataValidation/Private/EditorValidator_Material.cpp
+++ b/Engine/Plugins/Editor/DataValidation/Source/DataValidation/Private/EditorValidator_Material.cpp
@@ -11,6 +11,11 @@
 #include "Misc/DataValidation.h"
 #include "Widgets/Input/SComboBox.h"
 
+static TAutoConsoleVariable<bool> bCVarEdAllowMaterialValidator(
+	TEXT("Editor.EnableMaterialAssetValidator"),
+	false,
+	TEXT("Enables material asset validator in validation contexts where it's normally disabled (on save and from commandlet)."));
+
 UEditorValidator_Material::UEditorValidator_Material()
 	: Super()
 {
@@ -73,22 +78,20 @@ UEditorValidator_Material::UEditorValidator_Material()
 
 bool UEditorValidator_Material::CanValidateAsset_Implementation(const FAssetData& AssetData, UObject* InAsset, FDataValidationContext& InContext) const
 {
-	if (InContext.GetValidationUsecase() == EDataValidationUsecase::Save || InContext.GetValidationUsecase() == EDataValidationUsecase::Commandlet)
+	if (!bCVarEdAllowMaterialValidator.GetValueOnAnyThread() &&
+		(InContext.GetValidationUsecase() == EDataValidationUsecase::Save || InContext.GetValidationUsecase() == EDataValidationUsecase::Commandlet))
 	{
-		UE_LOG(LogContentValidation, Display, TEXT("skipping validation due to save or commandlet"));
 		return false;
 	}
 
 	if (ValidationPlatforms.IsEmpty())
 	{
-		UE_LOG(LogContentValidation, Display, TEXT("skipping validation due to empty validation platforms"));
 		return false;
 	}
 
 	const UMaterial* OriginalMaterial = Cast<UMaterial>(InAsset);
 	if (OriginalMaterial)
 	{
-		UE_LOG(LogContentValidation, Display, TEXT("enabling validation due to UMaterial"));
 		// We want to validate every UMaterial.
 		return true;
 	}
@@ -104,7 +107,6 @@ bool UEditorValidator_Material::CanValidateAsset_Implementation(const FAssetData
 			if (MaterialInstance->HasStaticParameters())
 			{
 				// We want to validate UMaterialInstance since it has a static parameter that could influence generated shader code.
-				UE_LOG(LogContentValidation, Display, TEXT("enabling validation due static parameters"));
 				return true;
 			}
 		}
@@ -124,7 +126,6 @@ EDataValidationResult UEditorValidator_Material::ValidateLoadedAsset_Implementat
 	UMaterial* Material = MaterialInstance ? MaterialInstance->GetMaterial() : DuplicateMaterial(OriginalMaterial);
 	if (!ensureAlways(OriginalMaterial) || !ensureAlways(Material))
 	{
-		UE_LOG(LogContentValidation, Display, TEXT("not validating because 1"));
 		return EDataValidationResult::NotValidated;
 	}
 
@@ -137,7 +138,6 @@ EDataValidationResult UEditorValidator_Material::ValidateLoadedAsset_Implementat
 
 		if (ensure(CurrentResource))
 		{
-			UE_LOG(LogContentValidation, Display, TEXT("caching shaders for '%s'"), *ValidationPlatform.ShaderPlatformName.ToString());
 			CurrentResource->CacheShaders(ValidationPlatform.ShaderPlatform);
 			ResourceToShaderPlatformNames.FindOrAdd(CurrentResource).Add(ValidationPlatform.ShaderPlatformName);
 		}
@@ -145,7 +145,6 @@ EDataValidationResult UEditorValidator_Material::ValidateLoadedAsset_Implementat
 
 	if (Resources.IsEmpty())
 	{
-		UE_LOG(LogContentValidation, Display, TEXT("not validating because 2"));
 		return EDataValidationResult::NotValidated;
 	}
 
@@ -167,23 +166,12 @@ EDataValidationResult UEditorValidator_Material::ValidateLoadedAsset_Implementat
 			UE_LOG(LogContentValidation, Error, TEXT("Failed to compile Material for platform %s due to '%s'."), *ShaderPlatformNameString, *ErrorText);
 			bCompileErrors = true;
 		}
-
-		UE_LOG(LogContentValidation, Display, TEXT("done with %s"), *ShaderPlatformNameString);
 	}
 
 	ResourceToShaderPlatformNames.Empty();
 
 	FMaterial::DeferredDeleteArray(Resources);
 
-	if (bCompileErrors)
-	{
-		UE_LOG(LogContentValidation, Display, TEXT("result - invalid"));
-	}
-	else
-	{
-		UE_LOG(LogContentValidation, Display, TEXT("result - ok"));
-	}
-
 	return !bCompileErrors ? EDataValidationResult::Valid : EDataValidationResult::Invalid;
 }
 
