Added a commend argument to only process the on disk asset in the DataValidationCommandlet.

diff --git a/Engine/Plugins/Editor/DataValidation/Source/DataValidation/Private/DataValidationCommandlet.cpp b/Engine/Plugins/Editor/DataValidation/Source/DataValidation/Private/DataValidationCommandlet.cpp
--- a/Engine/Plugins/Editor/DataValidation/Source/DataValidation/Private/DataValidationCommandlet.cpp
+++ b/Engine/Plugins/Editor/DataValidation/Source/DataValidation/Private/DataValidationCommandlet.cpp
@@ -54,7 +54,8 @@ bool UDataValidationCommandlet::ValidateData(const FString& FullCommandLine)
 	IAssetRegistry& AssetRegistry = AssetRegistryModule.Get();
 	AssetRegistry.SearchAllAssets(true /* bSynchronousSearch */);
 
-	bool bProjectOnly = !Switches.Contains(TEXT("includeengine"));
+	const bool bWithoutEngine = !Switches.Contains(TEXT("includeengine"));
+	const bool bIncludeOnlyOnDiskAssets = Switches.Contains(TEXT("IncludeOnlyOnDiskAssets"));
 
 	TArray<FAssetData> AssetDataList;
 	FString AssetTypeString;
@@ -73,15 +74,19 @@ bool UDataValidationCommandlet::ValidateData(const FString& FullCommandLine)
 				return false;
 			}
 		}
-		const bool bSearchSubClasses = true;
-		AssetRegistry.GetAssetsByClass(FTopLevelAssetPath(AssetTypeString), AssetDataList, bSearchSubClasses);
+	
+		FARFilter Filter;
+		Filter.ClassPaths.Add(FTopLevelAssetPath(AssetTypeString));
+		Filter.bRecursiveClasses = true;
+		Filter.bIncludeOnlyOnDiskAssets = bIncludeOnlyOnDiskAssets;
+		AssetRegistry.GetAssets(Filter, AssetDataList);
 	}
 	else
 	{
-		AssetRegistry.GetAllAssets(AssetDataList);
+		AssetRegistry.GetAllAssets(AssetDataList, bIncludeOnlyOnDiskAssets);
 	}
 
-	if (bProjectOnly)
+	if (bWithoutEngine)
 	{
 		FString EngineDir = FPaths::ConvertRelativePathToFull(FPaths::EngineDir());
 		AssetDataList.RemoveAll([&EngineDir](const FAssetData& AssetData)
