Manual asset validation will unload the assets afterward. This makes sure that repeated validations will not change behaviour, and also keeps the unneeded extra resources from the editor state.

diff --git a/Engine/Plugins/Editor/DataValidation/Source/DataValidation/Private/DataValidationModule.cpp b/Engine/Plugins/Editor/DataValidation/Source/DataValidation/Private/DataValidationModule.cpp
--- a/Engine/Plugins/Editor/DataValidation/Source/DataValidation/Private/DataValidationModule.cpp
+++ b/Engine/Plugins/Editor/DataValidation/Source/DataValidation/Private/DataValidationModule.cpp
@@ -330,6 +330,11 @@ void FDataValidationModule::ValidateAssets(const TArray<FAssetData>& SelectedAss
 		Settings.ValidationUsecase = InValidationUsecase;
 		Settings.MessageLogPageTitle = FText::Format(LOCTEXT("MessageLogPageTitle.ValidateSelectedAssets", "Asset Validation: {0}"), SelectedAssets.Num() == 1 ? FText::FromName(SelectedAssets[0].AssetName) : LOCTEXT("MultipleAssets", "multiple assets"));
 
+		if (InValidationUsecase == EDataValidationUsecase::Manual)
+		{
+			Settings.bUnloadAssetsLoadedForValidation = true;
+		}
+
 		FValidateAssetsResults Results;
 		EditorValidationSubsystem->ValidateAssetsWithSettings(bValidateDependencies ? DependentAssets.Array() : SelectedAssets, Settings, Results);
 	}
