Greatly speedup material validator by only translating material to HLSL first and only compiling if estimated amount of samplers exceed platform limits

diff --git a/Engine/Plugins/Editor/DataValidation/Source/DataValidation/Private/EditorValidator_Material.cpp b/Engine/Plugins/Editor/DataValidation/Source/DataValidation/Private/EditorValidator_Material.cpp
--- a/Engine/Plugins/Editor/DataValidation/Source/DataValidation/Private/EditorValidator_Material.cpp
+++ b/Engine/Plugins/Editor/DataValidation/Source/DataValidation/Private/EditorValidator_Material.cpp
@@ -11,6 +11,8 @@
 #include "Misc/DataValidation.h"
 #include "Widgets/Input/SComboBox.h"
 
+#define LOCTEXT_NAMESPACE "MaterialValidator"
+
 static TAutoConsoleVariable<bool> bCVarEdAllowMaterialValidator(
 	TEXT("Editor.EnableMaterialAssetValidator"),
 	false,
@@ -120,6 +122,88 @@ EDataValidationResult UEditorValidator_Material::ValidateLoadedAsset_Implementat
 	UMaterialInstance* OriginalMaterialInstance = Cast<UMaterialInstance>(InAsset); // Can be null
 	UMaterial* OriginalMaterial = OriginalMaterialInstance ? OriginalMaterialInstance->GetMaterial() : Cast<UMaterial>(InAsset);
 
+	if (!ensureAlways(OriginalMaterial))
+	{
+		return EDataValidationResult::NotValidated;
+	}
+
+	const bool bAllowCompilingShaders = GetDefault<UDataValidationSettings>()->bMaterialValidationAllowCompilingShaders;
+
+	bool bTranslationFailed = false;
+	TArray<int32> RetryWithCompilation;
+
+	TArray<FMaterialResource*> TranslationResources;
+	for (int32 ValidationPlatformIndex = 0; ValidationPlatformIndex < ValidationPlatforms.Num(); ++ValidationPlatformIndex)
+	{
+		const FShaderValidationPlatform& ValidationPlatform = ValidationPlatforms[ValidationPlatformIndex];
+
+		FMaterialResource* MaterialResource = FindOrCreateMaterialResource(TranslationResources, OriginalMaterial, OriginalMaterialInstance, ValidationPlatform.FeatureLevel, ValidationPlatform.MaterialQualityLevel);
+
+		FMaterialAnalysisResult Result;
+		OriginalMaterial->AnalyzeMaterialTranslationOutput(MaterialResource, ValidationPlatform.ShaderPlatform, Result);
+
+		if (!Result.bTranslationSuccess) // translation failed, report errors and early out
+		{
+			for (const FString& ErrorText: MaterialResource->GetCompileErrors())
+			{
+				AssetFails(InAsset, FText::Format(
+					LOCTEXT("MaterialValidator_TranslationError", "Failed to translate Material for platform {0} due to '{1}'."),
+					FText::FromName(ValidationPlatform.ShaderPlatformName), FText::FromString(ErrorText)
+				));
+			}
+
+			if (MaterialResource->GetCompileErrors().IsEmpty())
+			{
+				AssetFails(InAsset, LOCTEXT("MaterialValidator_TranslationFailedButNoError", "Material translation failed with no error."));
+			}
+
+			bTranslationFailed = true;
+		}
+		else // translation success, check if we run out of samplers
+		{
+			const uint32 MaxSamplers = FDataDrivenShaderPlatformInfo::GetMaxSamplers(ValidationPlatform.ShaderPlatform);
+
+			const bool bTooManyEstimatedSamplers = Result.EstimatedNumTextureSamplesVS > MaxSamplers || Result.EstimatedNumTextureSamplesPS > MaxSamplers;
+
+			if (bAllowCompilingShaders && bTooManyEstimatedSamplers)
+			{
+				RetryWithCompilation.Add(ValidationPlatformIndex);
+			}
+			else
+			{
+				// We want to show these to users only as warnings. Because they are based on estimated numbers and can generate false positive errors.
+
+				if (Result.EstimatedNumTextureSamplesVS > MaxSamplers)
+				{
+					AssetWarning(InAsset, FText::Format(
+						LOCTEXT("MaterialValidator_EstimatedVSOverPlatformLimit", "Estimated amount of VS samplers ({0}) is larger than supported on shader platform ({1}), shader will likely not compile for a shader platform '{2}'."),
+						Result.EstimatedNumTextureSamplesVS, MaxSamplers, FText::FromName(ValidationPlatform.ShaderPlatformName)
+					));
+				}
+
+				if (Result.EstimatedNumTextureSamplesPS > MaxSamplers)
+				{
+					AssetWarning(InAsset, FText::Format(
+						LOCTEXT("MaterialValidator_EstimatedPSOverPlatformLimit", "Estimated amount of PS samplers ({0}) is larger than supported on shader platform ({1}), shader will likely not compile for a shader platform '{2}'."),
+						Result.EstimatedNumTextureSamplesPS, MaxSamplers, FText::FromName(ValidationPlatform.ShaderPlatformName)
+					));
+				}
+			}
+		}
+	}
+
+	FMaterial::DeferredDeleteArray(TranslationResources);
+
+	if (bTranslationFailed) // early out if translation failed
+	{
+		return EDataValidationResult::Invalid;
+	}
+
+	if (RetryWithCompilation.IsEmpty()) // assume validation is ok and early out
+	{
+		return EDataValidationResult::Valid;
+	}
+
 	// Right now calling CacheShaders multiple times on the same UMaterial(Instance) will not work as we don't clean up stale compilation jobs.
 	// So to work around this we create a duplicate assets and compile cache shaders for them instead.
 	UMaterialInstance* MaterialInstance = DuplicateMaterialInstance(OriginalMaterialInstance);
@@ -129,12 +213,13 @@ EDataValidationResult UEditorValidator_Material::ValidateLoadedAsset_Implementat
 		return EDataValidationResult::NotValidated;
 	}
 
-	TArray<FMaterialResource*> Resources;
+	TArray<FMaterialResource*> CompilationResources;
 	TMap<FMaterialResource*, TArray<FName>> ResourceToShaderPlatformNames;
 
-	for (const FShaderValidationPlatform& ValidationPlatform: ValidationPlatforms)
+	for (const int32 ValidationPlatformIndex: RetryWithCompilation)
 	{
-		FMaterialResource* CurrentResource = FindOrCreateMaterialResource(Resources, Material, MaterialInstance, ValidationPlatform.FeatureLevel, ValidationPlatform.MaterialQualityLevel);
+		const FShaderValidationPlatform& ValidationPlatform = ValidationPlatforms[ValidationPlatformIndex];
+		FMaterialResource* CurrentResource = FindOrCreateMaterialResource(CompilationResources, Material, MaterialInstance, ValidationPlatform.FeatureLevel, ValidationPlatform.MaterialQualityLevel);
 
 		if (ensure(CurrentResource))
 		{
@@ -143,7 +228,7 @@ EDataValidationResult UEditorValidator_Material::ValidateLoadedAsset_Implementat
 		}
 	}
 
-	if (Resources.IsEmpty())
+	if (CompilationResources.IsEmpty())
 	{
 		return EDataValidationResult::NotValidated;
 	}
@@ -151,7 +236,7 @@ EDataValidationResult UEditorValidator_Material::ValidateLoadedAsset_Implementat
 	FAssetCompilingManager::Get().FinishAllCompilation();
 
 	bool bCompileErrors = false;
-	for (FMaterialResource* Resource: Resources)
+	for (FMaterialResource* Resource: CompilationResources)
 	{
 		if (!Resource->IsCompilationFinished())
 		{
@@ -163,16 +248,19 @@ EDataValidationResult UEditorValidator_Material::ValidateLoadedAsset_Implementat
 
 		for (const FString& ErrorText: Resource->GetCompileErrors())
 		{
-			UE_LOG(LogContentValidation, Error, TEXT("Failed to compile Material for platform %s due to '%s'."), *ShaderPlatformNameString, *ErrorText);
+			AssetFails(InAsset, FText::Format(
+				LOCTEXT("MaterialValidator_CompilationError", "Failed to compile Material for platform {0} due to '{1}'."),
+				FText::FromString(ShaderPlatformNameString), FText::FromString(ErrorText)
+			));
 			bCompileErrors = true;
 		}
 	}
 
 	ResourceToShaderPlatformNames.Empty();
 
-	FMaterial::DeferredDeleteArray(Resources);
+	FMaterial::DeferredDeleteArray(CompilationResources);
 
-	return !bCompileErrors ? EDataValidationResult::Valid : EDataValidationResult::Invalid;
+	return bCompileErrors ? EDataValidationResult::Invalid : EDataValidationResult::Valid;
 }
 
 UMaterial* UEditorValidator_Material::DuplicateMaterial(UMaterial* OriginalMaterial)
@@ -412,3 +500,5 @@ void FMaterialEditorValidationShaderPlatform::UnregisterCustomPropertyTypeLayout
 		PropertyModule->UnregisterCustomPropertyTypeLayout(CustomPropertyTypeLayoutName);
 	}
 }
+
+#undef LOCTEXT_NAMESPACE
diff --git a/Engine/Plugins/Editor/DataValidation/Source/DataValidation/Public/DataValidationSettings.h b/Engine/Plugins/Editor/DataValidation/Source/DataValidation/Public/DataValidationSettings.h
--- a/Engine/Plugins/Editor/DataValidation/Source/DataValidation/Public/DataValidationSettings.h
+++ b/Engine/Plugins/Editor/DataValidation/Source/DataValidation/Public/DataValidationSettings.h
@@ -30,6 +30,9 @@ public:
 	UPROPERTY(EditAnywhere, Config, Category="Data Validation")
 	bool bEnableMaterialValidation = true;
 
+	UPROPERTY(Config)
+	bool bMaterialValidationAllowCompilingShaders = true;
+
 	UPROPERTY(EditAnywhere, Config, Category="Data Validation", meta=(EditCondition="bEnableMaterialValidation"))
 	TArray<FMaterialEditorValidationPlatform> MaterialValidationPlatforms;
 };
diff --git a/Engine/Source/Runtime/Engine/Private/Materials/MaterialShared.cpp b/Engine/Source/Runtime/Engine/Private/Materials/MaterialShared.cpp
--- a/Engine/Source/Runtime/Engine/Private/Materials/MaterialShared.cpp
+++ b/Engine/Source/Runtime/Engine/Private/Materials/MaterialShared.cpp
@@ -4962,6 +4962,27 @@ void UMaterialInterface::AnalyzeMaterialCompilationInCallback(TFunctionRef<void
 		MaterialTranslator.bUsesAOMaterialMask ||
 		MaterialTranslator.bUsesLightmapUVs ||
 		MaterialTranslator.bUsesVertexPosition;
+	OutResult.EstimatedNumTextureSamplesVS = TempOutput.EstimatedNumTextureSamplesVS;
+	OutResult.EstimatedNumTextureSamplesPS = TempOutput.EstimatedNumTextureSamplesPS;
+#endif
+}
+
+void UMaterialInterface::AnalyzeMaterialTranslationOutput(FMaterialResource* MaterialResource, EShaderPlatform ShaderPlatform, FMaterialAnalysisResult& OutResult)
+{
+#if WITH_EDITORONLY_DATA
+	FStaticParameterSet StaticParamSet;
+	MaterialResource->GetStaticParameterSet(ShaderPlatform, StaticParamSet);
+
+	FMaterialCompilationOutput TempOutput;
+	FHLSLMaterialTranslator MaterialTranslator(MaterialResource, TempOutput, StaticParamSet, ShaderPlatform, MaterialResource->GetQualityLevel(), MaterialResource->GetFeatureLevel());
+
+	const EHLSLMaterialTranslatorResult TranslationResult = MaterialTranslator.Translate(true);
+
+	OutResult.bTranslationSuccess = TranslationResult == EHLSLMaterialTranslatorResult::Success;
+	OutResult.EstimatedNumTextureSamplesVS = TempOutput.EstimatedNumTextureSamplesVS;
+	OutResult.EstimatedNumTextureSamplesPS = TempOutput.EstimatedNumTextureSamplesPS;
+#else
+	OutResult = {};
 #endif
 }
 
diff --git a/Engine/Source/Runtime/Engine/Public/Materials/MaterialInterface.h b/Engine/Source/Runtime/Engine/Public/Materials/MaterialInterface.h
--- a/Engine/Source/Runtime/Engine/Public/Materials/MaterialInterface.h
+++ b/Engine/Source/Runtime/Engine/Public/Materials/MaterialInterface.h
@@ -229,6 +229,15 @@ struct FMaterialAnalysisResult
 
 	/** Whether any vertex data is used */
 	bool bRequiresVertexData = false;
+
+	/** If material translation was success, valid in AnalyzeMaterialTranslationOutput */
+	bool bTranslationSuccess = false;
+
+	/** Estimated amount of VS samplers used */
+	uint16 EstimatedNumTextureSamplesVS = 0;
+
+	/** Estimated amount of PS samplers used */
+	uint16 EstimatedNumTextureSamplesPS = 0;
 };
 
 USTRUCT()
@@ -1182,6 +1191,9 @@ public:
 	/** Return insight on what (e.g. texture coordinates, vertex data, etc) is used in the shader graph compiled by a callback */
 	ENGINE_API void AnalyzeMaterialCompilationInCallback(TFunctionRef<void (FMaterialCompiler*)> InCompilationCallback, FMaterialAnalysisResult& OutResult);
 
+	/** Return insight on errors or warnings generated during translation */
+	ENGINE_API void AnalyzeMaterialTranslationOutput(FMaterialResource* MaterialResource, EShaderPlatform ShaderPlatform, FMaterialAnalysisResult& OutResult);
+
 #if WITH_EDITOR
 	/** Checks to see if the given property references the texture */
 	ENGINE_API bool IsTextureReferencedByProperty(EMaterialProperty InProperty, const UTexture* InTexture);
