Gather statistics on which validators add more assets to be validated

diff --git a/Engine/Plugins/Editor/DataValidation/Source/DataValidation/Private/EditorValidatorSubsystem.cpp b/Engine/Plugins/Editor/DataValidation/Source/DataValidation/Private/EditorValidatorSubsystem.cpp
--- a/Engine/Plugins/Editor/DataValidation/Source/DataValidation/Private/EditorValidatorSubsystem.cpp
+++ b/Engine/Plugins/Editor/DataValidation/Source/DataValidation/Private/EditorValidatorSubsystem.cpp
@@ -22,6 +22,7 @@
 #include "ISourceControlModule.h"
 #include "LoadExternalObjectsForValidation.h"
 #include "Logging/MessageLog.h"
+#include "Logging/StructuredLog.h"
 #include "Misc/App.h"
 #include "Misc/DataValidation.h"
 #include "Misc/MessageDialog.h"
@@ -389,7 +390,8 @@ EDataValidationResult UEditorValidatorSubsystem::IsAssetValidWithContext(
 EDataValidationResult UEditorValidatorSubsystem::ValidateObjectInternal(
 	const FAssetData& InAssetData,
 	UObject* InObject,
-	FDataValidationContext& InContext) const
+	FDataValidationContext& InContext,
+	TMap<FTopLevelAssetPath, FValidatorStatistics>* ValidatorStatistics) const
 {
 	EDataValidationResult Result = EDataValidationResult::NotValidated;
 	
@@ -404,11 +406,15 @@ EDataValidationResult UEditorValidatorSubsystem::ValidateObjectInternal(
 			return Result;
 		}
 		
-		ForEachEnabledValidator([InObject, &InAssetData, &InContext, &Result](UEditorValidatorBase* Validator)
+		ForEachEnabledValidator([InObject, &InAssetData, &InContext, &Result, ValidatorStatistics](UEditorValidatorBase* Validator)
 		{
 			UE_LOG(LogContentValidation, Verbose, TEXT("Validating '%s' with '%s'..."), *InObject->GetPathName(), *Validator->GetClass()->GetName());
 			GInitRunaway(); // Reset runaway counter, as ValidateLoadedAsset may be implemented in a BP and could overflow the runaway count due to being called in a loop
 			EDataValidationResult NewResult = Validator->ValidateLoadedAsset(InAssetData, InObject, InContext);
+			if (NewResult != EDataValidationResult::NotValidated && ValidatorStatistics)
+			{
+				ValidatorStatistics->FindOrAdd(FTopLevelAssetPath(Validator->GetClass())).AssetsValidated++;
+			}
 			Result = CombineDataValidationResults(Result, NewResult);
 			return true;
 		});
@@ -449,7 +455,6 @@ EDataValidationResult UEditorValidatorSubsystem::ValidateAssetsInternal(
 		SlowTask.MakeDialog();
 	}
 	
-	UE_LOG(LogContentValidation, Display, TEXT("Starting to validate %d assets"), AssetDataList.Num());
 	UE_LOG(LogContentValidation, Log, TEXT("Enabled validators:"));
 	ForEachEnabledValidator([](UEditorValidatorBase* Validator)
 	{
@@ -500,16 +505,20 @@ EDataValidationResult UEditorValidatorSubsystem::ValidateAssetsInternal(
 	// Add any packages which contain those external objects to be validated
 	{
 		FDataValidationContext ValidationContext(false, InSettings.ValidationUsecase, {});
-		for (const TPair<FAssetData, TArray<FAssetData>>& Pair : AssetsToExternalObjects)
+		OutResults.NumExternalObjects = 0;
+		for (auto It = AssetsToExternalObjects.CreateIterator(); It; ++It)
 		{
+			const TPair<FAssetData, TArray<FAssetData>>& Pair = *It;
 			if (ShouldValidateAsset(Pair.Key, InSettings, ValidationContext))
 			{
 				AssetDataList.Add(Pair.Key);
+				OutResults.NumExternalObjects += Pair.Value.Num();
 			}
 			else
 			{
 				UE_LOG(LogContentValidation, Display, TEXT("Package %s (owner of some external objects) being skipped for validation."), 
 					*WriteToString<256>(Pair.Key.PackageName));
+				It.RemoveCurrent();
 			}
 		}
 		UE::DataValidation::AddAssetValidationMessages(DataValidationLog, ValidationContext);
@@ -521,6 +530,18 @@ EDataValidationResult UEditorValidatorSubsystem::ValidateAssetsInternal(
 
 	OutResults.NumRequested = AssetDataList.Num();
 	
+	UE_LOG(LogContentValidation, Display, TEXT("Starting to validate %d assets (%d associated objects such as actors)"), 
+		OutResults.NumRequested,
+		OutResults.NumExternalObjects);
+	
+	UE_LOGFMT_LOC(LogContentValidation, Log, "DataValidation.PreValidateStats.Preface", "Additional assets added for validation by {ValidatorCount} validators:",
+		("ValidatorCount", OutResults.ValidatorStatistics.Num()));
+	for (const TPair<FTopLevelAssetPath, FValidatorStatistics>& Pair : OutResults.ValidatorStatistics)
+	{
+		UE_LOGFMT(LogContentValidation, Log, "  {Validator} : {AssetsAddedForValidation}",
+			("Validator", Pair.Key), ("AssetsAddedForValidation", Pair.Value.AssetsAddedForValidation));
+	}
+	
 	EDataValidationResult Result = EDataValidationResult::NotValidated;
 	// Loaded assets ought to have the standalone flag and so not be garbage collected, but keep weak pointers here for memory safety
 	TArray<TWeakObjectPtr<UPackage>> PackagesToUnload;	
@@ -624,7 +645,7 @@ EDataValidationResult UEditorValidatorSubsystem::ValidateAssetsInternal(
 			UE::DataValidation::FScopedLoadExternalObjects ExternalObjectsLoader(LoadedAsset, ValidationContext, bLoadExternalObjects);
 
 			UE::DataValidation::FScopedLogMessageGatherer LogGatherer(InSettings.bCaptureLogsDuringValidation);
-			AssetResult = IsObjectValidWithContext(LoadedAsset, ValidationContext);
+			AssetResult = ValidateObjectInternal(Data, LoadedAsset, ValidationContext, &OutResults.ValidatorStatistics);
 			
 			// Associate any UE_LOG errors with this asset in the message log
 			TArray<FString> Warnings;
@@ -747,6 +768,14 @@ EDataValidationResult UEditorValidatorSubsystem::ValidateAssetsInternal(
 		}
 	}
 
+	UE_LOGFMT_LOC(LogContentValidation, Log, "DataValidation.PostValidateStats.Preface", "Validated asset counts for {ValidatorCount} validators:",
+		("ValidatorCount", OutResults.ValidatorStatistics.Num()));
+	for (const TPair<FTopLevelAssetPath, FValidatorStatistics>& Pair : OutResults.ValidatorStatistics)
+	{
+		UE_LOGFMT(LogContentValidation, Log, "  {Validator} : {AssetsValidated}",
+			("Validator", Pair.Key), ("AssetsValidated", Pair.Value.AssetsValidated));
+	}
+	
 	// Broadcast now that we're complete so other systems can go back to their previous state.
 	if (FEditorDelegates::OnPostAssetValidation.IsBound())
 	{
@@ -980,7 +1009,7 @@ EDataValidationResult UEditorValidatorSubsystem::ValidateChangelistsInternal(
 	for (UDataValidationChangelist* Changelist : Changelists)
 	{
 		FDataValidationContext ValidationContext(false, Settings.ValidationUsecase, {}); // No associated objects for changelist
-		GatherAssetsToValidateFromChangelist(Changelist, Settings, AssetsToValidate, ValidationContext);
+		GatherAssetsToValidateFromChangelist(Changelist, Settings, AssetsToValidate, ValidationContext, OutResults);
 		UE::DataValidation::AddAssetValidationMessages(DataValidationLog, ValidationContext);
 		DataValidationLog.Flush();
 	}
@@ -1017,7 +1046,8 @@ void UEditorValidatorSubsystem::GatherAssetsToValidateFromChangelist(
 	UDataValidationChangelist* 		InChangelist,
 	const FValidateAssetsSettings& 	Settings,
 	TSet<FAssetData>& 				OutAssets,
-	FDataValidationContext& 		InContext) const
+	FDataValidationContext& 		InContext,
+	FValidateAssetsResults& 		OutResults) const
 {
 	IAssetRegistry& AssetRegistry = IAssetRegistry::GetChecked();
 	
@@ -1029,31 +1059,44 @@ void UEditorValidatorSubsystem::GatherAssetsToValidateFromChangelist(
 	}
 	
 	// Gather assets requested by plugin/project validators 
-	ForEachEnabledValidator([this, InChangelist, &Settings, &InContext, &OutAssets](UEditorValidatorBase* Validator)
+	ForEachEnabledValidator([this, InChangelist, &Settings, &InContext, &OutAssets, &OutResults](UEditorValidatorBase* Validator)
 	{
+		int32 OldCount = OutAssets.Num();
 		TArray<FAssetData> NewAssets = Validator->GetAssetsToValidateFromChangelist(InChangelist, InContext);
 		for (const FAssetData& Asset : NewAssets)
 		{
 			// It's not strictly necessary to filter assets here but it makes logging simpler
 			if (ShouldValidateAsset(Asset, Settings, InContext))
 			{
-				UE_LOG(LogContentValidation, Display, TEXT("Asset validator %s adding %s to be validated from changelist %s."), *Validator->GetPathName(), *Asset.GetSoftObjectPath().ToString(), *InChangelist->Description.ToString());
-				OutAssets.Add(Asset);
+				bool bExisted = false;
+				OutAssets.Add(Asset, &bExisted);
+				UE_CLOG(!bExisted, LogContentValidation, Display, TEXT("Asset validator %s adding %s to be validated from changelist %s."), *Validator->GetPathName(), *Asset.GetSoftObjectPath().ToString(), *InChangelist->Description.ToString());
 			}
 		}
+		const int32 NumAssetsAdded = OutAssets.Num() - OldCount;
+		if (NumAssetsAdded)
+		{
+			FValidatorStatistics& Stats = OutResults.ValidatorStatistics.FindOrAdd(FTopLevelAssetPath(Validator->GetClass()));
+			Stats.AssetsAddedForValidation += NumAssetsAdded;
+		}
 		return true;
 	});
 	
 	if (Settings.bValidateReferencersOfDeletedAssets)
 	{
+		// Associate these with the validation subsystem base class - key type may need to change if we add more features here and want a breakdown
+		FValidatorStatistics& Stats = OutResults.ValidatorStatistics.FindOrAdd(FTopLevelAssetPath(UEditorValidatorSubsystem::StaticClass()));
+		int32 OldCount = OutAssets.Num();
 		TSet<FAssetData> AdditionalAssetsToValidate;
 		for (FName DeletedPackageName : InChangelist->DeletedPackageNames)
 		{
-			InContext.AddMessage(EMessageSeverity::Info, 
-				FText::Format(LOCTEXT("DataValidation.AddDeletedPackageReferencers", "Adding referencers of deleted package {0} to be validated"), FText::FromName(DeletedPackageName)));
-
 			TArray<FName> PackageReferencers;
 			AssetRegistry.GetReferencers(DeletedPackageName, PackageReferencers, UE::AssetRegistry::EDependencyCategory::Package);
+			if (PackageReferencers.Num() > 0)
+			{
+				InContext.AddMessage(EMessageSeverity::Info, 
+					FText::Format(LOCTEXT("DataValidation.AddDeletedPackageReferencers", "Adding referencers of deleted package {0} to be validated"), FText::FromName(DeletedPackageName)));
+			}
 			for (const FName& Referencer : PackageReferencers)
 			{
 				UE_LOG(LogContentValidation, Display, TEXT("Adding %s to to validated as it is a referencer of deleted asset %s"), *Referencer.ToString(), *DeletedPackageName.ToString());
@@ -1062,6 +1105,7 @@ void UEditorValidatorSubsystem::GatherAssetsToValidateFromChangelist(
 				OutAssets.Append(NewAssets);
 			}
 		}
+		Stats.AssetsAddedForValidation += OutAssets.Num() - OldCount;
 	}
 }
 
diff --git a/Engine/Plugins/Editor/DataValidation/Source/DataValidation/Public/EditorValidatorSubsystem.h b/Engine/Plugins/Editor/DataValidation/Source/DataValidation/Public/EditorValidatorSubsystem.h
--- a/Engine/Plugins/Editor/DataValidation/Source/DataValidation/Public/EditorValidatorSubsystem.h
+++ b/Engine/Plugins/Editor/DataValidation/Source/DataValidation/Public/EditorValidatorSubsystem.h
@@ -66,6 +66,20 @@ struct DATAVALIDATION_API FValidateAssetsDetails
 	TArray<FValidateAssetsExternalObject> ExternalObjects;
 };
 
+USTRUCT(BlueprintType)
+struct DATAVALIDATION_API FValidatorStatistics
+{
+	GENERATED_BODY()
+
+	/** Number of assets validated by this validator */
+	UPROPERTY(VisibleAnywhere, BlueprintReadOnly, Category="Validation")
+	int32 AssetsValidated;
+
+	/** Number of assets that this validator marked as needing validation e.g. because a dependency changed */
+	UPROPERTY(VisibleAnywhere, BlueprintReadOnly, Category="Validation")
+	int32 AssetsAddedForValidation;
+};
+
 USTRUCT(BlueprintType)
 struct DATAVALIDATION_API FValidateAssetsResults
 {
@@ -73,10 +87,14 @@ struct DATAVALIDATION_API FValidateAssetsResults
 
 	FValidateAssetsResults() = default;
 	
-	/** Total amount of assets that were gathered to validate. */
+	/** Total amount of assets that were gathered to validate, not including external objects such as actors. */
 	UPROPERTY(VisibleAnywhere, BlueprintReadOnly, Category = "Validation")
 	int NumRequested = 0;
 
+	/** Total amount of external objects to validate with their owning asset (e.g. actors with maps). */
+	UPROPERTY(VisibleAnywhere, BlueprintReadOnly, Category = "Validation")
+	int NumExternalObjects = 0;
+
 	/** Amount of tested assets */
 	UPROPERTY(VisibleAnywhere, BlueprintReadOnly, Category = "Validation")
 	int NumChecked = 0;
@@ -105,6 +123,10 @@ struct DATAVALIDATION_API FValidateAssetsResults
 	UPROPERTY(VisibleAnywhere, BlueprintReadOnly, Category = "Validation")
 	bool bAssetLimitReached = false;
 
+	/** Per-validator statistics on assets validated and assets added for validation */
+	UPROPERTY(VisibleAnywhere, BlueprintReadOnly, Category = "Validation")
+	TMap<FTopLevelAssetPath, FValidatorStatistics> ValidatorStatistics;
+
 	/** 
 	 * Per asset details
 	 * Indexed by object path
@@ -366,7 +388,8 @@ public:
 		UDataValidationChangelist* 		InChangelist,
 		const FValidateAssetsSettings& 	Settings,
 		TSet<FAssetData>& 				OutAssets,
-		FDataValidationContext& 		InContext) const;
+		FDataValidationContext& 		InContext,
+		FValidateAssetsResults& 		OutResults) const;
 
 	
 	/** 
@@ -461,7 +484,8 @@ protected:
 	EDataValidationResult ValidateObjectInternal(
 		const FAssetData& InAssetData,
 		UObject* InObject,
-		FDataValidationContext& InContext) const;
+		FDataValidationContext& InContext,
+		TMap<FTopLevelAssetPath, FValidatorStatistics>* ValidatorStatistics = nullptr) const;
 protected:
 	/**
 	 * Directories to ignore for data validation. Useful for test assets
