Add missing initializers for UStruct fields

diff --git a/Engine/Plugins/Editor/DataValidation/Source/DataValidation/Public/EditorValidatorSubsystem.h b/Engine/Plugins/Editor/DataValidation/Source/DataValidation/Public/EditorValidatorSubsystem.h
--- a/Engine/Plugins/Editor/DataValidation/Source/DataValidation/Public/EditorValidatorSubsystem.h
+++ b/Engine/Plugins/Editor/DataValidation/Source/DataValidation/Public/EditorValidatorSubsystem.h
@@ -73,11 +73,11 @@ struct DATAVALIDATION_API FValidatorStatistics
 
 	/** Number of assets validated by this validator */
 	UPROPERTY(VisibleAnywhere, BlueprintReadOnly, Category="Validation")
-	int32 AssetsValidated;
+	int32 AssetsValidated = 0;
 
 	/** Number of assets that this validator marked as needing validation e.g. because a dependency changed */
 	UPROPERTY(VisibleAnywhere, BlueprintReadOnly, Category="Validation")
-	int32 AssetsAddedForValidation;
+	int32 AssetsAddedForValidation = 0;
 };
 
 USTRUCT(BlueprintType)
