Fixed the ignored result from the blueprint validators.

diff --git a/Engine/Plugins/Editor/DataValidation/Source/DataValidation/Private/EditorValidatorBase.cpp b/Engine/Plugins/Editor/DataValidation/Source/DataValidation/Private/EditorValidatorBase.cpp
--- a/Engine/Plugins/Editor/DataValidation/Source/DataValidation/Private/EditorValidatorBase.cpp
+++ b/Engine/Plugins/Editor/DataValidation/Source/DataValidation/Private/EditorValidatorBase.cpp
@@ -35,6 +35,8 @@ EDataValidationResult UEditorValidatorBase::ValidateLoadedAsset(const FAssetData
 		EDataValidationResult K2Result = K2_ValidateLoadedAsset(Asset);
 		K2Result = CombineDataValidationResults(GetValidationResult(), K2Result);
 		ensureMsgf(K2Result != EDataValidationResult::NotValidated, TEXT("Validator %s did not return a validation result from BP validation path for asset %s"), *GetClass()->GetPathName(), *Asset->GetPathName());
+		
+		Result = CombineDataValidationResults(Result, K2Result);
 		bTriedToValidate = true;
 	}
 	
