Changelist Validation: fix crash when validating a loaded but not initialized partitioned world.

diff --git a/Engine/Plugins/Editor/DataValidation/Source/DataValidation/Private/WorldPartitionChangelistValidator.cpp b/Engine/Plugins/Editor/DataValidation/Source/DataValidation/Private/WorldPartitionChangelistValidator.cpp
--- a/Engine/Plugins/Editor/DataValidation/Source/DataValidation/Private/WorldPartitionChangelistValidator.cpp
+++ b/Engine/Plugins/Editor/DataValidation/Source/DataValidation/Private/WorldPartitionChangelistValidator.cpp
@@ -166,7 +166,7 @@ void UWorldPartitionChangelistValidator::ValidateActorsAndDataLayersFromChangeLi
 		}
 
 		UActorDescContainerInstance* ContainerInstance = nullptr;
-		if (InWorld != nullptr && InWorld->GetWorldPartition()->IsMainWorldPartition())
+		if (InWorld != nullptr && InWorld->GetWorldPartition()->IsInitialized() && InWorld->GetWorldPartition()->IsMainWorldPartition())
 		{
 			// World is Loaded reuse the ActorDescContainer (Can be Main world package or one of the Content Bundle packages associated with main world)
 			// Only reuse Main world partiton worlds because streaming generation expects the top level container to always be a Main container.
