EditorValidatorSubsystem: Do not popup the MessageLog if there are no additional errors/warnings.

The FMessageLog::Open() method inspects the entire message log and sees if any of them passes the given severity. If any pass, the log is opened. However, it inspects the entire log - not just the messages that were logged in the current context. By tracking the count of message pre/post, the unnecessary popping of the message log can be prevented.

#rnx

diff --git a/Engine/Plugins/Editor/DataValidation/Source/DataValidation/Private/EditorValidatorSubsystem.cpp b/Engine/Plugins/Editor/DataValidation/Source/DataValidation/Private/EditorValidatorSubsystem.cpp
--- a/Engine/Plugins/Editor/DataValidation/Source/DataValidation/Private/EditorValidatorSubsystem.cpp
+++ b/Engine/Plugins/Editor/DataValidation/Source/DataValidation/Private/EditorValidatorSubsystem.cpp
@@ -433,11 +433,22 @@ int32 UEditorValidatorSubsystem::ValidateAssetsWithSettings(
 {
 	FMessageLog DataValidationLog(InSettings.MessageLogName);
 	DataValidationLog.SetCurrentPage(InSettings.MessageLogPageTitle);
+
+	int32 NumMessagesPassingFilterPre = 0;
+	if (const EMessageSeverity::Type* Severity = InSettings.ShowMessageLogSeverity.GetPtrOrNull())
+	{
+		NumMessagesPassingFilterPre = DataValidationLog.NumMessages(*Severity);
+	}
+
 	ValidateAssetsInternal(DataValidationLog, TSet<FAssetData>{AssetDataList}, InSettings, OutResults);
 	
 	if (const EMessageSeverity::Type* Severity = InSettings.ShowMessageLogSeverity.GetPtrOrNull())
 	{
-		DataValidationLog.Open(*Severity, false);
+		int32 NumMessagesPassingFilterPost = DataValidationLog.NumMessages(*Severity);
+		if (NumMessagesPassingFilterPost > NumMessagesPassingFilterPre)
+		{
+			DataValidationLog.Open(*Severity, false);
+		}
 	}
 
 	return OutResults.NumWarnings + OutResults.NumInvalid;
@@ -923,6 +934,7 @@ void UEditorValidatorSubsystem::ValidateChangelistPreSubmit(
 	Settings.MessageLogPageTitle = FText::Format(LOCTEXT("MessageLogPageTitle.ValidateChangelist", "Changelist Validation: {0}"), FText::FromString(InChangelist->GetIdentifier()));
 	Settings.ShowMessageLogSeverity = { EMessageSeverity::Warning };
 	Settings.bCollectPerAssetDetails = true;
+	Settings.bShowIfNoFailures = false;
 
 	FValidateAssetsResults Results;
 	OutResult = ValidateChangelist(Changelist, Settings, Results);
@@ -952,6 +964,7 @@ void UEditorValidatorSubsystem::ValidateFilesPreSubmit(
 	Settings.MessageLogPageTitle = LOCTEXT("MessageLogPageTitle.ValidateFiles", "Files Validation");
 	Settings.ShowMessageLogSeverity = { EMessageSeverity::Warning };
 	Settings.bCollectPerAssetDetails = true;
+	Settings.bShowIfNoFailures = false;
 
 	FValidateAssetsResults Results;
 	OutResult = ValidateChangelist(Changelist, Settings, Results);
@@ -1007,6 +1020,12 @@ EDataValidationResult UEditorValidatorSubsystem::ValidateChangelistsInternal(
 	// Choose a specific message log page for this output, flushing in case other recursive calls also write to this log
 	DataValidationLog.SetCurrentPage(Settings.MessageLogPageTitle);
 
+	int32 NumMessagesPassingFilterPre = 0;
+	if (const EMessageSeverity::Type* Severity = Settings.ShowMessageLogSeverity.GetPtrOrNull())
+	{
+		NumMessagesPassingFilterPre = DataValidationLog.NumMessages(*Severity);
+	}
+
 	for (UDataValidationChangelist* CL : Changelists)
 	{
 		CL->AddToRoot();
@@ -1071,7 +1090,11 @@ EDataValidationResult UEditorValidatorSubsystem::ValidateChangelistsInternal(
 
 	if (const EMessageSeverity::Type* Severity = Settings.ShowMessageLogSeverity.GetPtrOrNull())
 	{
-		DataValidationLog.Open(*Severity, false);
+		int32 NumMessagesPassingFilterPost = DataValidationLog.NumMessages(*Severity);
+		if (NumMessagesPassingFilterPost > NumMessagesPassingFilterPre)
+		{
+			DataValidationLog.Open(*Severity, false);
+		}
 	}
 
 	return Result;
