MaterialCompiler validation mode to allow translation to evaluate multiple switch inputs in a single pass.
Quality Switch is the first to utilize this mode.

#rnx

diff --git a/Engine/Plugins/Editor/DataValidation/Source/DataValidation/Private/EditorValidator_Material.cpp b/Engine/Plugins/Editor/DataValidation/Source/DataValidation/Private/EditorValidator_Material.cpp
--- a/Engine/Plugins/Editor/DataValidation/Source/DataValidation/Private/EditorValidator_Material.cpp
+++ b/Engine/Plugins/Editor/DataValidation/Source/DataValidation/Private/EditorValidator_Material.cpp
@@ -130,6 +130,7 @@ EDataValidationResult UEditorValidator_Material::ValidateLoadedAsset_Implementat
 	}
 
 	const bool bAllowCompilingShaders = GetDefault<UDataValidationSettings>()->bMaterialValidationAllowCompilingShaders;
+	const bool bShowWarningsWhenNotCompilingShaders = GetDefault<UDataValidationSettings>()->bMaterialValidationShowWarningsWhenNotCompilingShaders;
 
 	bool bTranslationFailed = false;
 	TArray<int32> RetryWithCompilation;
@@ -142,7 +143,9 @@ EDataValidationResult UEditorValidator_Material::ValidateLoadedAsset_Implementat
 		FMaterialResource* MaterialResource = FindOrCreateMaterialResource(TranslationResources, OriginalMaterial, OriginalMaterialInstance, ValidationPlatform.ShaderPlatform, ValidationPlatform.MaterialQualityLevel);
 
 		FMaterialAnalysisResult Result;
-		OriginalMaterial->AnalyzeMaterialTranslationOutput(MaterialResource, ValidationPlatform.ShaderPlatform, Result);
+		// Translate in validation mode only when not allowing compiling shaders, otherwise the estimates can be higher and trigger unnecessary compilations.
+		const bool bValidationMode = !bAllowCompilingShaders;
+		OriginalMaterial->AnalyzeMaterialTranslationOutput(MaterialResource, ValidationPlatform.ShaderPlatform, bValidationMode, Result);
 
 		if (!Result.bTranslationSuccess) // translation failed, report errors and early out
 		{
@@ -171,7 +174,7 @@ EDataValidationResult UEditorValidator_Material::ValidateLoadedAsset_Implementat
 			{
 				RetryWithCompilation.Add(ValidationPlatformIndex);
 			}
-			else
+			else if (bShowWarningsWhenNotCompilingShaders)
 			{
 				// We want to show these to users only as warnings. Because they are based on estimated numbers and can generate false positive errors.
 
diff --git a/Engine/Plugins/Editor/DataValidation/Source/DataValidation/Public/DataValidationSettings.h b/Engine/Plugins/Editor/DataValidation/Source/DataValidation/Public/DataValidationSettings.h
--- a/Engine/Plugins/Editor/DataValidation/Source/DataValidation/Public/DataValidationSettings.h
+++ b/Engine/Plugins/Editor/DataValidation/Source/DataValidation/Public/DataValidationSettings.h
@@ -34,6 +34,9 @@ public:
 
 	UPROPERTY(Config)
 	bool bMaterialValidationAllowCompilingShaders = true;
+	
+	UPROPERTY(Config)
+	bool bMaterialValidationShowWarningsWhenNotCompilingShaders = true;
 
 	UPROPERTY(EditAnywhere, Config, Category="Data Validation", meta=(EditCondition="bEnableMaterialValidation"))
 	TArray<FMaterialEditorValidationPlatform> MaterialValidationPlatforms;
diff --git a/Engine/Source/Runtime/Engine/Private/Materials/MaterialExpressions.cpp b/Engine/Source/Runtime/Engine/Private/Materials/MaterialExpressions.cpp
--- a/Engine/Source/Runtime/Engine/Private/Materials/MaterialExpressions.cpp
+++ b/Engine/Source/Runtime/Engine/Private/Materials/MaterialExpressions.cpp
@@ -9277,6 +9277,21 @@ int32 UMaterialExpressionQualitySwitch::Compile(class FMaterialCompiler* Compile
 		return Compiler->Errorf(TEXT("Quality switch missing default input"));
 	}
 
+	// Evaluate all inputs during validation to minimize validation permutations.
+	if (Compiler->IsValidating())
+	{
+		Default.Compile(Compiler);
+		for (uint32 QualityLevelToCompile = 0; QualityLevelToCompile < EMaterialQualityLevel::Num; ++QualityLevelToCompile)
+		{
+			check(QualityLevelToCompile < UE_ARRAY_COUNT(Inputs));
+			FExpressionInput QualityInputTraced = Inputs[QualityLevelToCompile].GetTracedInput();
+			if (QualityInputTraced.Expression && QualityInputTraced.Expression != DefaultTraced.Expression)
+			{
+				Inputs[QualityLevelToCompile].Compile(Compiler);
+			}
+		}
+	}
+
 	FExpressionInput* EffectiveInput = GetEffectiveInput(Compiler);
 	check(EffectiveInput);
 	return EffectiveInput->Compile(Compiler);
diff --git a/Engine/Source/Runtime/Engine/Private/Materials/MaterialShared.cpp b/Engine/Source/Runtime/Engine/Private/Materials/MaterialShared.cpp
--- a/Engine/Source/Runtime/Engine/Private/Materials/MaterialShared.cpp
+++ b/Engine/Source/Runtime/Engine/Private/Materials/MaterialShared.cpp
@@ -5072,6 +5072,11 @@ void UMaterialInterface::AnalyzeMaterialCompilationInCallback(TFunctionRef<void
 }
 
 void UMaterialInterface::AnalyzeMaterialTranslationOutput(FMaterialResource* MaterialResource, EShaderPlatform ShaderPlatform, FMaterialAnalysisResult& OutResult)
+{
+	UMaterialInterface::AnalyzeMaterialTranslationOutput(MaterialResource, ShaderPlatform, false, OutResult);
+}
+
+void UMaterialInterface::AnalyzeMaterialTranslationOutput(FMaterialResource* MaterialResource, EShaderPlatform ShaderPlatform, bool ValidationMode, FMaterialAnalysisResult& OutResult)
 {
 #if WITH_EDITORONLY_DATA
 	FStaticParameterSet StaticParamSet;
@@ -5079,6 +5084,7 @@ void UMaterialInterface::AnalyzeMaterialTranslationOutput(FMaterialResource* Mat
 
 	FMaterialCompilationOutput TempOutput;
 	FHLSLMaterialTranslator MaterialTranslator(MaterialResource, TempOutput, StaticParamSet, ShaderPlatform, MaterialResource->GetQualityLevel(), MaterialResource->GetFeatureLevel());
+	MaterialTranslator.EnableValidationMode(ValidationMode);
 
 	const EHLSLMaterialTranslatorResult TranslationResult = MaterialTranslator.Translate(true);
 
diff --git a/Engine/Source/Runtime/Engine/Public/MaterialCompiler.h b/Engine/Source/Runtime/Engine/Public/MaterialCompiler.h
--- a/Engine/Source/Runtime/Engine/Public/MaterialCompiler.h
+++ b/Engine/Source/Runtime/Engine/Public/MaterialCompiler.h
@@ -796,6 +796,10 @@ public:
 	virtual bool IsCurrentlyCompilingForPreviousFrame() const { return false; }
 	virtual bool IsDevelopmentFeatureEnabled(const FName& FeatureName) const { return true; }
 
+	// Validation mode to allow different handling of material expressions during translation.
+	void EnableValidationMode(bool bEnable) { bIsValidating = bEnable; }
+	bool IsValidating() const { return bIsValidating; }
+
 protected:
 	FString GetSubstrateSharedLocalBasisIndexMacroInner(const FSubstrateRegisteredSharedLocalBasis& SharedLocalBasis, uint32 Mode)
 	{
@@ -811,6 +815,7 @@ private:
 	ESubstrateMaterialExport SubstrateMaterialExport = SME_None;
 	ESubstrateMaterialExportContext SubstrateMaterialExportContext = SMEC_Opaque;
 	uint8 SubstrateMaterialExportLegacyBlendMode = 0;
+	bool bIsValidating = false;
 };
 
 /** 
diff --git a/Engine/Source/Runtime/Engine/Public/Materials/MaterialInterface.h b/Engine/Source/Runtime/Engine/Public/Materials/MaterialInterface.h
--- a/Engine/Source/Runtime/Engine/Public/Materials/MaterialInterface.h
+++ b/Engine/Source/Runtime/Engine/Public/Materials/MaterialInterface.h
@@ -1204,7 +1204,9 @@ public:
 	ENGINE_API void AnalyzeMaterialCompilationInCallback(TFunctionRef<void (FMaterialCompiler*)> InCompilationCallback, FMaterialAnalysisResult& OutResult);
 
 	/** Return insight on errors or warnings generated during translation */
+	UE_DEPRECATED(5.7, "Please use AnalyzeMaterialTranslationOutput with ValidationMode argument")
 	ENGINE_API void AnalyzeMaterialTranslationOutput(FMaterialResource* MaterialResource, EShaderPlatform ShaderPlatform, FMaterialAnalysisResult& OutResult);
+	ENGINE_API void AnalyzeMaterialTranslationOutput(FMaterialResource* MaterialResource, EShaderPlatform ShaderPlatform, bool ValidationMode, FMaterialAnalysisResult& OutResult);
 
 #if WITH_EDITOR
 	/** Checks to see if the given property references the texture */
