Add additional logging for validator registration

diff --git a/Engine/Plugins/Editor/DataValidation/Source/DataValidation/Private/EditorValidatorSubsystem.cpp b/Engine/Plugins/Editor/DataValidation/Source/DataValidation/Private/EditorValidatorSubsystem.cpp
--- a/Engine/Plugins/Editor/DataValidation/Source/DataValidation/Private/EditorValidatorSubsystem.cpp
+++ b/Engine/Plugins/Editor/DataValidation/Source/DataValidation/Private/EditorValidatorSubsystem.cpp
@@ -169,7 +169,8 @@ void UEditorValidatorSubsystem::RegisterNativeValidators()
 			&& FPackageName::IsScriptPackage(FNameBuilder(ValidatorClass->GetPackage()->GetFName())))
 		{
 			UEditorValidatorBase* Validator = NewObject<UEditorValidatorBase>(GetTransientPackage(), ValidatorClass);
-			AddValidator(Validator);
+			UE_LOGFMT(LogContentValidation, Log, "Adding native validator {ValidatorClass}", FTopLevelAssetPath(ValidatorClass));
+			AddValidator_Internal(Validator);
 		}
 	}
 
@@ -201,7 +202,8 @@ void UEditorValidatorSubsystem::RegisterBlueprintValidators()
 			// Skip any C++ classes, as well as anything that has already been registered by this point
 			if (!Validators.Contains(ValidatorClass) && !FPackageName::IsScriptPackage(FNameBuilder(ValidatorClass.GetPackageName())))
 			{
-				AddValidator(ValidatorClass);
+				UE_LOGFMT(LogContentValidation, Log, "Adding blueprint validator {ValidatorClass}", FTopLevelAssetPath(ValidatorClass));
+				AddValidator_Internal(ValidatorClass);
 			}
 		}
 
@@ -256,6 +258,30 @@ bool UEditorValidatorSubsystem::ShouldValidateOnSave(bool bProceduralSave) const
 }
 
 void UEditorValidatorSubsystem::AddValidator(UEditorValidatorBase* InValidator)
+{
+	UE_LOGFMT(LogContentValidation, Log, "Explicitly adding custom validator {ValidatorClass}", FTopLevelAssetPath(InValidator->GetClass()));
+	AddValidator_Internal(InValidator);
+}
+
+void UEditorValidatorSubsystem::AddValidator(const FTopLevelAssetPath& InValidatorClass)
+{
+	UE_LOGFMT(LogContentValidation, Log, "Explicitly adding custom validator {ValidatorClass}", InValidatorClass);
+	AddValidator_Internal(InValidatorClass);
+}
+
+void UEditorValidatorSubsystem::RemoveValidator(UEditorValidatorBase* InValidator)
+{
+	UE_LOGFMT(LogContentValidation, Log, "Explicitly removing custom validator {ValidatorClass}", FTopLevelAssetPath(InValidator->GetClass()));
+	RemoveValidator_Internal(InValidator);
+}
+
+void UEditorValidatorSubsystem::RemoveValidator(const FTopLevelAssetPath& InValidatorClass)
+{
+	UE_LOGFMT(LogContentValidation, Log, "Explicitly removing custom validator {ValidatorClass}", InValidatorClass);
+	RemoveValidator_Internal(InValidatorClass);
+}
+
+void UEditorValidatorSubsystem::AddValidator_Internal(UEditorValidatorBase* InValidator)
 {
 	if (InValidator)
 	{
@@ -265,7 +291,7 @@ void UEditorValidatorSubsystem::AddValidator(UEditorValidatorBase* InValidator)
 	}
 }
 
-void UEditorValidatorSubsystem::AddValidator(const FTopLevelAssetPath& InValidatorClass)
+void UEditorValidatorSubsystem::AddValidator_Internal(FTopLevelAssetPath InValidatorClass)
 {
 	if (InValidatorClass.IsValid())
 	{
@@ -274,7 +300,7 @@ void UEditorValidatorSubsystem::AddValidator(const FTopLevelAssetPath& InValidat
 	}
 }
 
-void UEditorValidatorSubsystem::RemoveValidator(UEditorValidatorBase* InValidator)
+void UEditorValidatorSubsystem::RemoveValidator_Internal(UEditorValidatorBase* InValidator)
 {
 	if (InValidator)
 	{
@@ -282,7 +308,7 @@ void UEditorValidatorSubsystem::RemoveValidator(UEditorValidatorBase* InValidato
 	}
 }
 
-void UEditorValidatorSubsystem::RemoveValidator(const FTopLevelAssetPath& InValidatorClass)
+void UEditorValidatorSubsystem::RemoveValidator_Internal(FTopLevelAssetPath InValidatorClass)
 {
 	if (InValidatorClass.IsValid())
 	{
@@ -1188,12 +1214,18 @@ void UEditorValidatorSubsystem::OnNativeModulesChanged(FName InModuleName, EModu
 
 void UEditorValidatorSubsystem::OnAssetsAdded(TConstArrayView<FAssetData> InAssets)
 {
-	OnAssetsAddedOrRemoved(InAssets, [this](const FTopLevelAssetPath& BPGC) { AddValidator(BPGC); });
+	OnAssetsAddedOrRemoved(InAssets, [this](const FTopLevelAssetPath& BPGC) { 
+		UE_LOGFMT(LogContentValidation, Log, "Adding validator from newly added asset {ValidatorClass}", BPGC);
+		AddValidator_Internal(BPGC); 
+	});
 }
 
 void UEditorValidatorSubsystem::OnAssetsRemoved(TConstArrayView<FAssetData> InAssets)
 {
-	OnAssetsAddedOrRemoved(InAssets, [this](const FTopLevelAssetPath& BPGC) { RemoveValidator(BPGC); });
+	OnAssetsAddedOrRemoved(InAssets, [this](const FTopLevelAssetPath& BPGC) {
+		UE_LOGFMT(LogContentValidation, Log, "Removing validator from removed asset {ValidatorClass}", BPGC);
+		RemoveValidator_Internal(BPGC); 
+	});
 }
 
 void UEditorValidatorSubsystem::OnAssetsAddedOrRemoved(TConstArrayView<FAssetData> InAssets, TFunctionRef<void(const FTopLevelAssetPath& BPGC)> Callback)
@@ -1277,7 +1309,8 @@ void UEditorValidatorSubsystem::UpdateValidators()
 				&& ModulePackageNames.Contains(ValidatorClass->GetPackage()->GetFName()))
 			{
 				UEditorValidatorBase* Validator = NewObject<UEditorValidatorBase>(GetTransientPackage(), ValidatorClass);
-				AddValidator(Validator);
+				UE_LOGFMT(LogContentValidation, Log, "Adding native validator {ValidatorClass} from newly loaded module", FTopLevelAssetPath(ValidatorClass));
+				AddValidator_Internal(Validator);
 			}
 		}
 
@@ -1305,6 +1338,7 @@ void UEditorValidatorSubsystem::UpdateValidators()
 
 				if (ValidatorClass && !ValidatorClass->HasAnyClassFlags(CLASS_Abstract))
 				{
+					UE_LOGFMT(LogContentValidation, Log, "Creating validator instance {ValidatorClass}", FTopLevelAssetPath(ValidatorClass));
 					*ValidatorInstance = NewObject<UEditorValidatorBase>(GetTransientPackage(), ValidatorClass);
 				}
 			}
diff --git a/Engine/Plugins/Editor/DataValidation/Source/DataValidation/Public/EditorValidatorSubsystem.h b/Engine/Plugins/Editor/DataValidation/Source/DataValidation/Public/EditorValidatorSubsystem.h
--- a/Engine/Plugins/Editor/DataValidation/Source/DataValidation/Public/EditorValidatorSubsystem.h
+++ b/Engine/Plugins/Editor/DataValidation/Source/DataValidation/Public/EditorValidatorSubsystem.h
@@ -472,6 +472,13 @@ protected:
 	 */
 	UE_API void RemoveValidator(const FTopLevelAssetPath& InValidatorClass);
 
+private:
+	void AddValidator_Internal(UEditorValidatorBase* InValidator);
+	void AddValidator_Internal(FTopLevelAssetPath InValidatorClass);
+	void RemoveValidator_Internal(UEditorValidatorBase* InValidator);
+	void RemoveValidator_Internal(FTopLevelAssetPath InValidatorClass);
+protected:
+
 	/** 
 	 * Validate a set of assets, adding the results to the log/output of a higher level task (e.g. changelist validation).
 	 */
