Add additional null checks before logging

diff --git a/Engine/Plugins/Editor/DataValidation/Source/DataValidation/Private/EditorValidatorSubsystem.cpp b/Engine/Plugins/Editor/DataValidation/Source/DataValidation/Private/EditorValidatorSubsystem.cpp
--- a/Engine/Plugins/Editor/DataValidation/Source/DataValidation/Private/EditorValidatorSubsystem.cpp
+++ b/Engine/Plugins/Editor/DataValidation/Source/DataValidation/Private/EditorValidatorSubsystem.cpp
@@ -259,8 +259,11 @@ bool UEditorValidatorSubsystem::ShouldValidateOnSave(bool bProceduralSave) const
 
 void UEditorValidatorSubsystem::AddValidator(UEditorValidatorBase* InValidator)
 {
-	UE_LOGFMT(LogContentValidation, Log, "Explicitly adding custom validator {ValidatorClass}", FTopLevelAssetPath(InValidator->GetClass()));
-	AddValidator_Internal(InValidator);
+	if (InValidator)
+	{
+		UE_LOGFMT(LogContentValidation, Log, "Explicitly adding custom validator {ValidatorClass}", FTopLevelAssetPath(InValidator->GetClass()));
+		AddValidator_Internal(InValidator);
+	}
 }
 
 void UEditorValidatorSubsystem::AddValidator(const FTopLevelAssetPath& InValidatorClass)
@@ -271,8 +274,11 @@ void UEditorValidatorSubsystem::AddValidator(const FTopLevelAssetPath& InValidat
 
 void UEditorValidatorSubsystem::RemoveValidator(UEditorValidatorBase* InValidator)
 {
-	UE_LOGFMT(LogContentValidation, Log, "Explicitly removing custom validator {ValidatorClass}", FTopLevelAssetPath(InValidator->GetClass()));
-	RemoveValidator_Internal(InValidator);
+	if (InValidator)
+	{
+		UE_LOGFMT(LogContentValidation, Log, "Explicitly removing custom validator {ValidatorClass}", FTopLevelAssetPath(InValidator->GetClass()));
+		RemoveValidator_Internal(InValidator);
+	}
 }
 
 void UEditorValidatorSubsystem::RemoveValidator(const FTopLevelAssetPath& InValidatorClass)
