Add slow task when loading validators

diff --git a/Engine/Plugins/Editor/DataValidation/Source/DataValidation/Private/EditorValidatorSubsystem.cpp b/Engine/Plugins/Editor/DataValidation/Source/DataValidation/Private/EditorValidatorSubsystem.cpp
--- a/Engine/Plugins/Editor/DataValidation/Source/DataValidation/Private/EditorValidatorSubsystem.cpp
+++ b/Engine/Plugins/Editor/DataValidation/Source/DataValidation/Private/EditorValidatorSubsystem.cpp
@@ -1326,8 +1326,11 @@ void UEditorValidatorSubsystem::UpdateValidators()
 	// Add any new validators from pending Blueprint classes
 	if (ValidatorClassesPendingLoad.Num() > 0)
 	{
+		FScopedSlowTask SlowTask(ValidatorClassesPendingLoad.Num(), LOCTEXT("DataValidation.LoadingValidatorsTask", "Loading validators..."));
+		SlowTask.MakeDialog();
 		for (const FTopLevelAssetPath& ValidatorClassPendingLoad : ValidatorClassesPendingLoad)
 		{
+			SlowTask.EnterProgressFrame(1.0f);
 			if (TObjectPtr<UEditorValidatorBase>* ValidatorInstance = Validators.Find(ValidatorClassPendingLoad);
 				ensure(ValidatorInstance))
 			{
