Consolidate common errors across shader platforms for material validator.

#rnx

diff --git a/Engine/Plugins/Editor/DataValidation/Source/DataValidation/Private/EditorValidator_Material.cpp b/Engine/Plugins/Editor/DataValidation/Source/DataValidation/Private/EditorValidator_Material.cpp
--- a/Engine/Plugins/Editor/DataValidation/Source/DataValidation/Private/EditorValidator_Material.cpp
+++ b/Engine/Plugins/Editor/DataValidation/Source/DataValidation/Private/EditorValidator_Material.cpp
@@ -134,6 +134,7 @@ EDataValidationResult UEditorValidator_Material::ValidateLoadedAsset_Implementat
 
 	bool bTranslationFailed = false;
 	TArray<int32> RetryWithCompilation;
+	TMap<FString, TArray<FName>> TranslationErrorMap;
 
 	TArray<FMaterialResource*> TranslationResources;
 	for (int32 ValidationPlatformIndex = 0; ValidationPlatformIndex < ValidationPlatforms.Num(); ++ValidationPlatformIndex)
@@ -151,10 +152,7 @@ EDataValidationResult UEditorValidator_Material::ValidateLoadedAsset_Implementat
 		{
 			for (const FString& ErrorText: MaterialResource->GetCompileErrors())
 			{
-				AssetFails(InAsset, FText::Format(
-					LOCTEXT("MaterialValidator_TranslationError", "Failed to translate Material for platform {0} due to '{1}'."),
-					FText::FromName(ValidationPlatform.ShaderPlatformName), FText::FromString(ErrorText)
-				));
+				TranslationErrorMap.FindOrAdd(ErrorText).AddUnique(ValidationPlatform.ShaderPlatformName);
 			}
 
 			if (MaterialResource->GetCompileErrors().IsEmpty())
@@ -197,6 +195,26 @@ EDataValidationResult UEditorValidator_Material::ValidateLoadedAsset_Implementat
 		}
 	}
 
+	if (!TranslationErrorMap.IsEmpty())
+	{
+		for (const TPair<FString, TArray<FName>>& ErrorAndPlatforms: TranslationErrorMap)
+		{
+			TStringBuilder<128> PlatformsString;
+			for (const FName& Platform: ErrorAndPlatforms.Value)
+			{
+				if (PlatformsString.Len() != 0)
+				{
+					PlatformsString += TEXT(", ");
+				}
+				PlatformsString += Platform.ToString();
+			}
+			AssetFails(InAsset, FText::Format(
+				LOCTEXT("MaterialValidator_TranslationError", "Failed to translate Material for platform(s) {0} due to '{1}'."),
+				FText::FromStringView(PlatformsString.ToView()), FText::FromString(ErrorAndPlatforms.Key)
+			));
+		}
+	}
+	
 	FMaterial::DeferredDeleteArray(TranslationResources);
 
 	if (bTranslationFailed) // early out if translation failed
