Add structure log context for assets being validated

diff --git a/Engine/Plugins/Editor/DataValidation/Source/DataValidation/Private/EditorValidatorSubsystem.cpp b/Engine/Plugins/Editor/DataValidation/Source/DataValidation/Private/EditorValidatorSubsystem.cpp
--- a/Engine/Plugins/Editor/DataValidation/Source/DataValidation/Private/EditorValidatorSubsystem.cpp
+++ b/Engine/Plugins/Editor/DataValidation/Source/DataValidation/Private/EditorValidatorSubsystem.cpp
@@ -432,6 +432,7 @@ EDataValidationResult UEditorValidatorSubsystem::ValidateObjectInternal(
 	TMap<FTopLevelAssetPath, FValidatorStatistics>* ValidatorStatistics) const
 {
 	EDataValidationResult Result = EDataValidationResult::NotValidated;
+	UE_LOG_CONTEXT("AssetBeingValidated", InAssetData.GetSoftObjectPath());
 	
 	if (ensure(InObject) && ensure(InAssetData.IsValid()))
 	{
@@ -663,6 +664,7 @@ EDataValidationResult UEditorValidatorSubsystem::ValidateAssetsInternal(
 		{
 			UE_LOG(LogContentValidation, Log, TEXT("Loading asset %s for validation. After the validation we will %sunload it"), *Data.ToSoftObjectPath().ToString(), InSettings.bUnloadAssetsLoadedForValidation ? TEXT("") : TEXT("not "));
 			UE::DataValidation::FScopedLogMessageGatherer LogGatherer(InSettings.bCaptureAssetLoadLogs);
+			UE_LOG_CONTEXT("AssetBeingValidated", Data.GetSoftObjectPath());
 			
 			// Don't pass the flag to load all external objects. Later we'll load external objects which are also part of the to-validate
 			// set if the settings allow for it.
