Add post validation virtual function for editor validators.
Validators can do clean up work and add non-asset specific tokenized messages.
#rnx
#Synced-CL 48135599

diff --git a/Engine/Plugins/Editor/DataValidation/Source/DataValidation/Private/EditorValidatorSubsystem.cpp b/Engine/Plugins/Editor/DataValidation/Source/DataValidation/Private/EditorValidatorSubsystem.cpp
--- a/Engine/Plugins/Editor/DataValidation/Source/DataValidation/Private/EditorValidatorSubsystem.cpp
+++ b/Engine/Plugins/Editor/DataValidation/Source/DataValidation/Private/EditorValidatorSubsystem.cpp
@@ -830,6 +830,17 @@ EDataValidationResult UEditorValidatorSubsystem::ValidateAssetsInternal(
 		}
 	}
 
+	ForEachEnabledValidator([&](UEditorValidatorBase* Validator)
+	{
+		Validator->OnPostAssetValidation(OutResults.ValidatorMessages);
+		return true;
+	});
+
+	for (TSharedRef<FTokenizedMessage> const& Message : OutResults.ValidatorMessages)
+	{
+		DataValidationLog.AddMessage(Message);
+	}
+	
 	UE_LOGFMT_LOC(LogContentValidation, Log, "DataValidation.PostValidateStats.Preface", "Validated asset counts for {ValidatorCount} validators:",
 		("ValidatorCount", OutResults.ValidatorStatistics.Num()));
 	for (const TPair<FTopLevelAssetPath, FValidatorStatistics>& Pair : OutResults.ValidatorStatistics)
@@ -997,6 +1008,11 @@ void UEditorValidatorSubsystem::ValidateChangelistPreSubmit(
 			OutResult.AllMessages.Add(Message->ToText());
 		}
 	}
+
+	for (TSharedRef<FTokenizedMessage> const& Message : Results.ValidatorMessages)
+	{
+		OutResult.AllMessages.Add(Message->ToText());
+	}
 }
 
 void UEditorValidatorSubsystem::ValidateFilesPreSubmit(
@@ -1033,6 +1049,11 @@ void UEditorValidatorSubsystem::ValidateFilesPreSubmit(
 			}
 		}
 	}
+
+	for (TSharedRef<FTokenizedMessage> const& Message : Results.ValidatorMessages)
+	{
+		OutResult.AllMessages.Add(Message->ToText());
+	}
 }
 
 EDataValidationResult UEditorValidatorSubsystem::ValidateChangelist(
diff --git a/Engine/Plugins/Editor/DataValidation/Source/DataValidation/Public/EditorValidatorBase.h b/Engine/Plugins/Editor/DataValidation/Source/DataValidation/Public/EditorValidatorBase.h
--- a/Engine/Plugins/Editor/DataValidation/Source/DataValidation/Public/EditorValidatorBase.h
+++ b/Engine/Plugins/Editor/DataValidation/Source/DataValidation/Public/EditorValidatorBase.h
@@ -57,6 +57,9 @@ public:
 		return {};
 	}
 	
+	/** Override this to do any state cleanup or add additional validation messages after all per asset validation is complete. */
+	virtual void OnPostAssetValidation(TArray<TSharedRef<FTokenizedMessage>>& OutMessages) {}
+
 	/** Override this to determine whether or not you can use this validator given this usecase */
 	UFUNCTION(BlueprintNativeEvent, Category = "Asset Validation", DisplayName="Can Validate")
 	UE_API bool K2_CanValidate(const EDataValidationUsecase InUsecase) const;
diff --git a/Engine/Plugins/Editor/DataValidation/Source/DataValidation/Public/EditorValidatorSubsystem.h b/Engine/Plugins/Editor/DataValidation/Source/DataValidation/Public/EditorValidatorSubsystem.h
--- a/Engine/Plugins/Editor/DataValidation/Source/DataValidation/Public/EditorValidatorSubsystem.h
+++ b/Engine/Plugins/Editor/DataValidation/Source/DataValidation/Public/EditorValidatorSubsystem.h
@@ -138,6 +138,9 @@ struct FValidateAssetsResults
 	 */
 	UPROPERTY(VisibleAnywhere, BlueprintReadOnly, Category = "Validation")
 	TMap<FString, FValidateAssetsDetails> AssetsDetails;
+
+	/** Additional validator tokenized messages not specific to a single asset. */
+	TArray<TSharedRef<FTokenizedMessage>> ValidatorMessages;
 };
 
 USTRUCT(BlueprintType)
