PR #13987: Refactored the DataValidationCommandlet to support customization via child classes

#Synced-CL 48966252

diff --git a/Engine/Plugins/Editor/DataValidation/Source/DataValidation/Private/DataValidationCommandlet.cpp b/Engine/Plugins/Editor/DataValidation/Source/DataValidation/Private/DataValidationCommandlet.cpp
--- a/Engine/Plugins/Editor/DataValidation/Source/DataValidation/Private/DataValidationCommandlet.cpp
+++ b/Engine/Plugins/Editor/DataValidation/Source/DataValidation/Private/DataValidationCommandlet.cpp
@@ -28,64 +28,162 @@ int32 UDataValidationCommandlet::Main(const FString& FullCommandLine)
 	check(GEditor);
 
 	UE_LOG(LogDataValidation, Log, TEXT("--------------------------------------------------------------------------------------------"));
-	UE_LOG(LogDataValidation, Log, TEXT("Running DataValidation Commandlet"));
+	UE_LOG(LogDataValidation, Log, TEXT("Running %s Commandlet"), *GetClass()->GetName());
 
 	// validate data
-	if (!ValidateData(FullCommandLine))
+	if (!ValidateDataImpl(FullCommandLine))
 	{
 		UE_LOG(LogDataValidation, Warning, TEXT("Errors occurred while validating data"));
 		return 2; // return something other than 1 for error since the engine will return 1 if any other system (possibly unrelated) logged errors during execution.
 	}
 
-	UE_LOG(LogDataValidation, Log, TEXT("Successfully finished running DataValidation Commandlet"));
+	UE_LOG(LogDataValidation, Log, TEXT("Successfully finished running %s Commandlet"), *GetClass()->GetName());
 	UE_LOG(LogDataValidation, Log, TEXT("--------------------------------------------------------------------------------------------"));
 	return 0;
 }
 
 //static
 bool UDataValidationCommandlet::ValidateData(const FString& FullCommandLine)
+{
+	const UEditorValidatorSubsystem* EditorValidationSubsystem = GEditor->GetEditorSubsystem<UEditorValidatorSubsystem>();
+	check(EditorValidationSubsystem);
+
+	// Attempt to load the data validation commandlet class from the EditorValidationSubsystem's configured
+	// GetDataValidationCommandletClassName property, but if it is not specified then fall back on the UDataValidationCommandlet.
+	UClass* CommandletClass = UDataValidationCommandlet::StaticClass();
+	if (!EditorValidationSubsystem->GetDataValidationCommandletClassName().IsEmpty())
+	{
+		UClass* ConfiguredClass = FindObject<UClass>(nullptr, *EditorValidationSubsystem->GetDataValidationCommandletClassName());
+		if (ConfiguredClass == nullptr)
+		{
+			UE_LOG(LogDataValidation, Error, TEXT("Unable to load Data Validation Commandlet Class: %s"), *EditorValidationSubsystem->GetDataValidationCommandletClassName());
+			return false;
+		}
+		else if (!ConfiguredClass->IsChildOf(UDataValidationCommandlet::StaticClass()))
+		{
+			UE_LOG(LogDataValidation, Error, TEXT("Data Validation Commandlet Class Name is not a UDataValidationCommandlet class: %s"), *EditorValidationSubsystem->GetDataValidationCommandletClassName());
+			return false;
+		}
+		else
+		{
+			CommandletClass = ConfiguredClass;
+		}
+	}
+
+	TStrongObjectPtr<UDataValidationCommandlet> Commandlet{ NewObject<UDataValidationCommandlet>(GetTransientPackage(), CommandletClass) };
+	if (Commandlet)
+	{
+		return Commandlet->ValidateDataImpl(FullCommandLine);
+	}
+	else
+	{
+		UE_LOG(LogDataValidation, Error, TEXT("Data Validation Commandlet can be instanciated. The UClass might be abstract. Commandlet Class Name: %s"), *EditorValidationSubsystem->GetDataValidationCommandletClassName());
+		return false;
+	}
+}
+
+bool UDataValidationCommandlet::ValidateDataImpl(const FString& FullCommandLine)
+{
+	ProcessCommandLine(FullCommandLine);
+
+	FAssetRegistryModule& AssetRegistryModule = FModuleManager::LoadModuleChecked<FAssetRegistryModule>(AssetRegistryConstants::ModuleName);
+	IAssetRegistry& AssetRegistry = AssetRegistryModule.Get();
+	WaitForAssetRegistry();
+
+	TArray<FAssetData> AssetDataList;
+	if (!GetAssetsToValidate(AssetRegistry, AssetDataList))
+	{
+		return false;
+	}
+
+	FilterAssetsToValidate(AssetDataList);
+
+	if (!GEditor->IsInitialized() &&
+		ShouldLoadDefaultEditorModules(AssetDataList))
+	{
+		GEditor->LoadDefaultEditorModules();
+	}
+
+	FValidateAssetsSettings Settings;
+	SetupValidationSettings(Settings);
+
+	FValidateAssetsResults Results;
+	ValidateAssets(AssetDataList, Settings, Results);
+
+	return ProcessValidationResults(Results);
+}
+
+bool UDataValidationCommandlet::ShouldLoadDefaultEditorModules(const TArray<FAssetData>& AssetDataList)
+{
+	// Check if we have some BP validator that were created using an editor utility.
+	// Those Editor Utilities Validator might have an dependency to an editor module that is loaded during the editor initialization.
+	const FTopLevelAssetPath EditorUtilityClassPath = UEditorUtilityBlueprint::StaticClass()->GetClassPathName();
+	const FString EditorValidatorBaseClassExportPath = FObjectPropertyBase::GetExportPath(UEditorValidatorBase::StaticClass());
+	return AssetDataList.ContainsByPredicate([EditorUtilityClassPath, &EditorValidatorBaseClassExportPath](const FAssetData& AssetData)
+		{
+			if (AssetData.AssetClassPath == EditorUtilityClassPath)
+			{
+				if (AssetData.TagsAndValues.ContainsKeyValue(FBlueprintTags::NativeParentClassPath, EditorValidatorBaseClassExportPath))
+				{
+					return true;
+				}
+			}
+
+			return false;
+		});
+}
+
+void UDataValidationCommandlet::ProcessCommandLine(const FString& FullCommandLine)
 {
 	TArray<FString> Tokens;
 	TArray<FString> Switches;
 	TMap<FString, FString> Params;
 	ParseCommandLine(*FullCommandLine, Tokens, Switches, Params);
 
-	FAssetRegistryModule& AssetRegistryModule = FModuleManager::LoadModuleChecked<FAssetRegistryModule>(AssetRegistryConstants::ModuleName);
-	IAssetRegistry& AssetRegistry = AssetRegistryModule.Get();
-	AssetRegistry.SearchAllAssets(true /* bSynchronousSearch */);
+	ProcessCommandLine(Tokens, Switches, Params);
+}
 
-	const bool bWithoutEngine = !Switches.Contains(TEXT("includeengine"));
-	const bool bIncludeOnlyOnDiskAssets = Switches.Contains(TEXT("IncludeOnlyOnDiskAssets"));
+void UDataValidationCommandlet::ProcessCommandLine(const TArray<FString>& Tokens, const TArray<FString>& Switches, const TMap<FString, FString>& Params)
+{
+	AssetTypeFilterString = Params.FindRef(TEXT("AssetType"));
+	bIncludeOnlyOnDiskAssets = Switches.Contains(TEXT("IncludeOnlyOnDiskAssets"));
+	bWithoutEngine = !Switches.Contains(TEXT("includeengine"));
+}
 
-	TArray<FAssetData> AssetDataList;
-	FString AssetTypeString;
-	if (FParse::Value(*FullCommandLine, TEXT("AssetType="), AssetTypeString) && !AssetTypeString.IsEmpty())
+bool UDataValidationCommandlet::GetAssetsToValidate(IAssetRegistry& AssetRegistry, TArray<FAssetData>& OutAssetDataList)
+{
+	if (!AssetTypeFilterString.IsEmpty())
 	{
-		if (FPackageName::IsShortPackageName(AssetTypeString))
+		if (FPackageName::IsShortPackageName(AssetTypeFilterString))
 		{
-			UClass* Class = FindFirstObject<UClass>(*AssetTypeString, EFindFirstObjectOptions::EnsureIfAmbiguous);
+			UClass* Class = FindFirstObject<UClass>(*AssetTypeFilterString, EFindFirstObjectOptions::EnsureIfAmbiguous);
 			if (Class)
 			{
-				AssetTypeString = Class->GetPathName();
+				AssetTypeFilterString = Class->GetPathName();
 			}
 			else
 			{
-				UE_LOG(LogDataValidation, Error, TEXT("Unable to resolve class path name given short name: \"%s\""), *AssetTypeString);
+				UE_LOG(LogDataValidation, Error, TEXT("Unable to resolve class path name given short name: \"%s\""), *AssetTypeFilterString);
 				return false;
 			}
 		}
 	
 		FARFilter Filter;
-		Filter.ClassPaths.Add(FTopLevelAssetPath(AssetTypeString));
+		Filter.ClassPaths.Add(FTopLevelAssetPath(AssetTypeFilterString));
 		Filter.bRecursiveClasses = true;
 		Filter.bIncludeOnlyOnDiskAssets = bIncludeOnlyOnDiskAssets;
-		AssetRegistry.GetAssets(Filter, AssetDataList);
+		AssetRegistry.GetAssets(Filter, OutAssetDataList);
 	}
 	else
 	{
-		AssetRegistry.GetAllAssets(AssetDataList, bIncludeOnlyOnDiskAssets);
+		AssetRegistry.GetAllAssets(OutAssetDataList, bIncludeOnlyOnDiskAssets);
 	}
 
+	return true;
+}
+
+void UDataValidationCommandlet::FilterAssetsToValidate(TArray<FAssetData>& AssetDataList)
+{
 	if (bWithoutEngine)
 	{
 		FString EngineDir = FPaths::ConvertRelativePathToFull(FPaths::EngineDir());
@@ -107,46 +205,33 @@ bool UDataValidationCommandlet::ValidateData(const FString& FullCommandLine)
 				return FPathViews::IsParentPathOf(EngineDir, FileName);
 			});
 	}
+}
 
-
-	if (!GEditor->IsInitialized())
-	{
-		// Check if we have some BP validator that were created using an editor utility
-		const FTopLevelAssetPath EditorUtilityClassPath = UEditorUtilityBlueprint::StaticClass()->GetClassPathName();
-		const FString EditorValidatorBaseClassExportPath = FObjectPropertyBase::GetExportPath(UEditorValidatorBase::StaticClass());
-		const bool bHasAnEditorUtilityDataValidator = AssetDataList.ContainsByPredicate([EditorUtilityClassPath, &EditorValidatorBaseClassExportPath](const FAssetData& AssetData)
-			{
-				if (AssetData.AssetClassPath == EditorUtilityClassPath)
-				{
-					if (AssetData.TagsAndValues.ContainsKeyValue(FBlueprintTags::NativeParentClassPath, EditorValidatorBaseClassExportPath))
-					{
-						return true;
-					}
-				}
-
-				return false;
-			});
-
-		if (bHasAnEditorUtilityDataValidator)
-		{
-			// Those Editor Utilities Validator might have an dependency to an editor module that is loaded during the editor initialization.
-			GEditor->LoadDefaultEditorModules();
-		}
-	}
-
-
-	UEditorValidatorSubsystem* EditorValidationSubsystem = GEditor->GetEditorSubsystem<UEditorValidatorSubsystem>();
-	check(EditorValidationSubsystem);
-
-	FValidateAssetsSettings Settings;
-	FValidateAssetsResults Results;
-
+void UDataValidationCommandlet::SetupValidationSettings(FValidateAssetsSettings& Settings)
+{
 	Settings.bSkipExcludedDirectories = true;
 	Settings.bShowIfNoFailures = true;
 	Settings.ValidationUsecase = EDataValidationUsecase::Commandlet;
+}
 
-	EditorValidationSubsystem->ValidateAssetsWithSettings(AssetDataList, Settings, Results);
+void UDataValidationCommandlet::ValidateAssets(
+	const TArray<FAssetData>& AssetDataList,
+	const FValidateAssetsSettings& InSettings,
+	FValidateAssetsResults& OutResults)
+{
+	UEditorValidatorSubsystem* EditorValidationSubsystem = GEditor->GetEditorSubsystem<UEditorValidatorSubsystem>();
+	check(EditorValidationSubsystem);
+	EditorValidationSubsystem->ValidateAssetsWithSettings(AssetDataList, InSettings, OutResults);
+}
 
+bool UDataValidationCommandlet::ProcessValidationResults(FValidateAssetsResults& Results)
+{
+	// Only return false if the commandlet couldn't run the validation and not if there was error in found by the validation
 	return true;
 }
 
+void UDataValidationCommandlet::WaitForAssetRegistry()
+{
+	constexpr bool bSynchronousSearch = true;
+	IAssetRegistry::GetChecked().SearchAllAssets(bSynchronousSearch);
+}
diff --git a/Engine/Plugins/Editor/DataValidation/Source/DataValidation/Public/DataValidationCommandlet.h b/Engine/Plugins/Editor/DataValidation/Source/DataValidation/Public/DataValidationCommandlet.h
--- a/Engine/Plugins/Editor/DataValidation/Source/DataValidation/Public/DataValidationCommandlet.h
+++ b/Engine/Plugins/Editor/DataValidation/Source/DataValidation/Public/DataValidationCommandlet.h
@@ -22,8 +22,32 @@ public:
 	UE_API virtual int32 Main(const FString& FullCommandLine) override;
 	// End UCommandlet Interface
 
-	// do the validation without creating a commandlet
+	// Run the commandlet validation logic without actually launching the editor in a commandlet mode.
 	static UE_API bool ValidateData(const FString& FullCommandLine);
+
+protected:
+	UE_API bool ValidateDataImpl(const FString& FullCommandLine);
+
+	// Parses the commandline and processes any found tokens/switch/parameters.
+	UE_API void ProcessCommandLine(const FString& FullCommandLine);
+	UE_API virtual void ProcessCommandLine(const TArray<FString>& Tokens, const TArray<FString>& Switches, const TMap<FString, FString>& Params);
+
+	UE_API virtual bool GetAssetsToValidate(class IAssetRegistry& AssetRegistry, TArray<FAssetData>& OutAssetDataList);
+	UE_API virtual void FilterAssetsToValidate(TArray<FAssetData>& AssetDataList);
+
+	UE_API virtual bool ShouldLoadDefaultEditorModules(const TArray<FAssetData>& AssetDataList);
+
+	UE_API virtual void SetupValidationSettings(struct FValidateAssetsSettings& Settings);
+	UE_API virtual void ValidateAssets(const TArray<FAssetData>& AssetDataList, const struct FValidateAssetsSettings& InSettings, struct FValidateAssetsResults& OutResults);
+	UE_API virtual bool ProcessValidationResults(struct FValidateAssetsResults& Results);
+
+	/** Checks if the AssetRegistry is performing its initial startup scan, and if so waits for the scan to complete. */
+	static UE_API void WaitForAssetRegistry();
+
+protected:
+	FString AssetTypeFilterString;
+	bool bIncludeOnlyOnDiskAssets = false;
+	bool bWithoutEngine = true;
 };
 
 #undef UE_API
diff --git a/Engine/Plugins/Editor/DataValidation/Source/DataValidation/Public/EditorValidatorSubsystem.h b/Engine/Plugins/Editor/DataValidation/Source/DataValidation/Public/EditorValidatorSubsystem.h
--- a/Engine/Plugins/Editor/DataValidation/Source/DataValidation/Public/EditorValidatorSubsystem.h
+++ b/Engine/Plugins/Editor/DataValidation/Source/DataValidation/Public/EditorValidatorSubsystem.h
@@ -416,6 +416,12 @@ public:
 	// Retrieve all assets matching the given filter from the asset registry and then recursively resolve redirectors to produce a single set of assets.
 	UE_API TArray<FAssetData> GetAssetsResolvingRedirectors(FARFilter& InFilter);
 
+	/**
+	 * Gets the full class object path of the UDataValidationCommandlet class to use when the Editor needs to instantiate
+	 * the commandlet to run validation on the entire project. If empty, UDataValidationCommandlet will be used.
+	 */
+	const FString& GetDataValidationCommandletClassName() const { return DataValidationCommandletClassName; }
+
 
 protected:
 	UE_API bool ShouldValidateOnSave(bool bProceduralSave) const;
@@ -512,6 +518,13 @@ protected:
 	UPROPERTY(config)
 	TArray<FDirectoryPath> ExcludedDirectories;
 
+	/**
+	 * The full class object path of the UDataValidationCommandlet class to use when the Editor needs to instantiate
+	 * the commandlet to run validation on the entire project. If empty, UDataValidationCommandlet will be used.
+	 */
+	UPROPERTY(config)
+	FString DataValidationCommandletClassName;
+
 	/**
 	 * Whether it should validate assets on save inside the editor
 	 */
