Don't promote logged warnings to errors by default, as this causes issues for validators that explicitly log warnings

#rnx
#Synced-CL 49183599

diff --git a/Engine/Plugins/Editor/DataValidation/Source/DataValidation/Public/EditorValidatorSubsystem.h b/Engine/Plugins/Editor/DataValidation/Source/DataValidation/Public/EditorValidatorSubsystem.h
--- a/Engine/Plugins/Editor/DataValidation/Source/DataValidation/Public/EditorValidatorSubsystem.h
+++ b/Engine/Plugins/Editor/DataValidation/Source/DataValidation/Public/EditorValidatorSubsystem.h
@@ -192,7 +192,7 @@ struct FValidateAssetsSettings
 	
 	/** If true, captured log warnings during validation are added to the validation results as errors (requires bCaptureLogsDuringValidation) */
 	UPROPERTY(EditAnywhere, BlueprintReadWrite, Category = "Validation")
-	bool bCaptureWarningsDuringValidationAsErrors = true;
+	bool bCaptureWarningsDuringValidationAsErrors = false;
 
 	/** Maximum number of assets to attempt to validate */
 	UPROPERTY(EditAnywhere, BlueprintReadWrite, Category = "Validation")
