Added utility to make data validation page titles pseudo-unique

#rnx
#Synced-CL 49217210

diff --git a/Engine/Plugins/Editor/DataValidation/Source/DataValidation/Private/EditorValidatorHelpers.cpp b/Engine/Plugins/Editor/DataValidation/Source/DataValidation/Private/EditorValidatorHelpers.cpp
--- a/Engine/Plugins/Editor/DataValidation/Source/DataValidation/Private/EditorValidatorHelpers.cpp
+++ b/Engine/Plugins/Editor/DataValidation/Source/DataValidation/Private/EditorValidatorHelpers.cpp
@@ -82,6 +82,17 @@ void AddAssetValidationMessages(
     }
 }
 
+FText MakeTimestampedMessageLogPageTitle(const FText& BaseTitle)
+{
+	return MakeTimestampedMessageLogPageTitle(BaseTitle, FDateTime::UtcNow());
+}
+
+FText MakeTimestampedMessageLogPageTitle(const FText& BaseTitle, const FDateTime& UtcTimestamp)
+{
+	static const FTextFormat TimestampedPageTitleFmt = NSLOCTEXT("DataValidation", "TimestampedPageTitleFmt", "[{0}]\t{1}");
+	return FText::Format(TimestampedPageTitleFmt, FText::AsDateTime(UtcTimestamp, EDateTimeStyle::Short, EDateTimeStyle::Default), BaseTitle);
+}
+
 struct FLogMessageGathererImpl : public FOutputDevice
 {
     FLogMessageGathererImpl()
diff --git a/Engine/Plugins/Editor/DataValidation/Source/DataValidation/Public/EditorValidatorHelpers.h b/Engine/Plugins/Editor/DataValidation/Source/DataValidation/Public/EditorValidatorHelpers.h
--- a/Engine/Plugins/Editor/DataValidation/Source/DataValidation/Public/EditorValidatorHelpers.h
+++ b/Engine/Plugins/Editor/DataValidation/Source/DataValidation/Public/EditorValidatorHelpers.h
@@ -116,6 +116,10 @@ namespace UE::DataValidation
         FMessageLog& Log,
         EMessageSeverity::Type Severity,
         TConstArrayView<FText> Messages);
+
+	/** Prepend the given UTC timestamp (or UtcNow if none provided) to the given base title to provide a pseudo-unique title, eg) "[Timestamp] BaseTitle" */
+	UE_API FText MakeTimestampedMessageLogPageTitle(const FText& BaseTitle);
+	UE_API FText MakeTimestampedMessageLogPageTitle(const FText& BaseTitle, const FDateTime& UtcTimestamp);
 }
 
 #undef UE_API
