Made data validation page titles pseudo-unique to avoid the confusion where you continue to see old issues when re-validating

#rnx
#Synced-CL 49217210

diff --git a/Engine/Plugins/Editor/DataValidation/Source/DataValidation/Private/DataValidationModule.cpp b/Engine/Plugins/Editor/DataValidation/Source/DataValidation/Private/DataValidationModule.cpp
--- a/Engine/Plugins/Editor/DataValidation/Source/DataValidation/Private/DataValidationModule.cpp
+++ b/Engine/Plugins/Editor/DataValidation/Source/DataValidation/Private/DataValidationModule.cpp
@@ -328,7 +328,7 @@ void FDataValidationModule::ValidateAssets(const TArray<FAssetData>& SelectedAss
 		Settings.bSkipExcludedDirectories = false;
 		Settings.bShowIfNoFailures = true;
 		Settings.ValidationUsecase = InValidationUsecase;
-		Settings.MessageLogPageTitle = FText::Format(LOCTEXT("MessageLogPageTitle.ValidateSelectedAssets", "Asset Validation: {0}"), SelectedAssets.Num() == 1 ? FText::FromName(SelectedAssets[0].AssetName) : LOCTEXT("MultipleAssets", "multiple assets"));
+		Settings.MessageLogPageTitle = UE::DataValidation::MakeTimestampedMessageLogPageTitle(FText::Format(LOCTEXT("MessageLogPageTitle.ValidateSelectedAssets", "Asset Validation: {0}"), SelectedAssets.Num() == 1 ? FText::FromName(SelectedAssets[0].AssetName) : LOCTEXT("MultipleAssets", "multiple assets")));
 
 		if (InValidationUsecase == EDataValidationUsecase::Manual)
 		{
diff --git a/Engine/Plugins/Editor/DataValidation/Source/DataValidation/Private/EditorValidatorSubsystem.cpp b/Engine/Plugins/Editor/DataValidation/Source/DataValidation/Private/EditorValidatorSubsystem.cpp
--- a/Engine/Plugins/Editor/DataValidation/Source/DataValidation/Private/EditorValidatorSubsystem.cpp
+++ b/Engine/Plugins/Editor/DataValidation/Source/DataValidation/Private/EditorValidatorSubsystem.cpp
@@ -907,9 +907,11 @@ void UEditorValidatorSubsystem::ValidateOnSave(TArray<FAssetData> AssetDataList,
 	Settings.bShowIfNoFailures = false;
 	Settings.ValidationUsecase = EDataValidationUsecase::Save;
 	Settings.bLoadAssetsForValidation = false;
-	Settings.MessageLogPageTitle = FText::Format(LOCTEXT("MessageLogPageTitle.ValidateSavedAssets", "Asset Save: {0}"), SavedAsset);
+	Settings.MessageLogPageTitle = UE::DataValidation::MakeTimestampedMessageLogPageTitle(FText::Format(LOCTEXT("MessageLogPageTitle.ValidateSavedAssets", "Asset Save: {0}"), SavedAsset));
 
-	if (ValidateAssetsWithSettings(AssetDataList, Settings, Results) > 0)
+	const int32 NumAssetsWithIssues = ValidateAssetsWithSettings(AssetDataList, Settings, Results);
+
+	if (NumAssetsWithIssues > 0 && !Settings.ShowMessageLogSeverity.IsSet()) // Only show the toast if ValidateAssetsWithSettings didn't already open the message log
 	{
 		FMessageLog DataValidationLog(Settings.MessageLogName);
 
@@ -988,7 +990,7 @@ void UEditorValidatorSubsystem::ValidateChangelistPreSubmit(
 	FValidateAssetsSettings Settings;
 	Settings.ValidationUsecase = EDataValidationUsecase::PreSubmit;
 	Settings.bLoadAssetsForValidation = GetDefault<UDataValidationSettings>()->bLoadAssetsWhenValidatingChangelists;
-	Settings.MessageLogPageTitle = FText::Format(LOCTEXT("MessageLogPageTitle.ValidateChangelist", "Changelist Validation: {0}"), FText::FromString(InChangelist->GetIdentifier()));
+	Settings.MessageLogPageTitle = UE::DataValidation::MakeTimestampedMessageLogPageTitle(FText::Format(LOCTEXT("MessageLogPageTitle.ValidateChangelist", "Changelist Validation: {0}"), FText::FromString(InChangelist->GetIdentifier())));
 	Settings.ShowMessageLogSeverity = { EMessageSeverity::Warning };
 	Settings.bCollectPerAssetDetails = true;
 	Settings.bShowIfNoFailures = false;
@@ -1026,7 +1028,7 @@ void UEditorValidatorSubsystem::ValidateFilesPreSubmit(
 	FValidateAssetsSettings Settings;
 	Settings.ValidationUsecase = EDataValidationUsecase::PreSubmit;
 	Settings.bLoadAssetsForValidation = GetDefault<UDataValidationSettings>()->bLoadAssetsWhenValidatingChangelists;
-	Settings.MessageLogPageTitle = LOCTEXT("MessageLogPageTitle.ValidateFiles", "Files Validation");
+	Settings.MessageLogPageTitle = UE::DataValidation::MakeTimestampedMessageLogPageTitle(LOCTEXT("MessageLogPageTitle.ValidateFiles", "Files Validation"));
 	Settings.ShowMessageLogSeverity = { EMessageSeverity::Warning };
 	Settings.bCollectPerAssetDetails = true;
 	Settings.bShowIfNoFailures = false;
