Updated FDataValidationContext to store the MessageLogName so validators can reference it.

[RN] Updated FDataValidationContext to store the MessageLogName so validators can reference it.
#Synced-CL 49280314

diff --git a/Engine/Plugins/Editor/DataValidation/Source/DataValidation/Private/EditorValidatorHelpers.cpp b/Engine/Plugins/Editor/DataValidation/Source/DataValidation/Private/EditorValidatorHelpers.cpp
--- a/Engine/Plugins/Editor/DataValidation/Source/DataValidation/Private/EditorValidatorHelpers.cpp
+++ b/Engine/Plugins/Editor/DataValidation/Source/DataValidation/Private/EditorValidatorHelpers.cpp
@@ -13,8 +13,6 @@
 namespace UE::DataValidation
 {
 
-const FName MessageLogName("AssetCheck");
-
 void AddAssetValidationMessages(
     FMessageLog& Log,
     const FDataValidationContext& ValidationContext)
diff --git a/Engine/Plugins/Editor/DataValidation/Source/DataValidation/Private/EditorValidatorSubsystem.cpp b/Engine/Plugins/Editor/DataValidation/Source/DataValidation/Private/EditorValidatorSubsystem.cpp
--- a/Engine/Plugins/Editor/DataValidation/Source/DataValidation/Private/EditorValidatorSubsystem.cpp
+++ b/Engine/Plugins/Editor/DataValidation/Source/DataValidation/Private/EditorValidatorSubsystem.cpp
@@ -565,7 +565,7 @@ EDataValidationResult UEditorValidatorSubsystem::ValidateAssetsInternal(
 	
 	// Add any packages which contain those external objects to be validated
 	{
-		FDataValidationContext ValidationContext(false, InSettings.ValidationUsecase, {});
+		FDataValidationContext ValidationContext(false, InSettings.ValidationUsecase, {}, InSettings.MessageLogName);
 		OutResults.NumExternalObjects = 0;
 		for (auto It = AssetsToExternalObjects.CreateIterator(); It; ++It)
 		{
@@ -658,7 +658,7 @@ EDataValidationResult UEditorValidatorSubsystem::ValidateAssetsInternal(
 			ValidationExternalObjects = *ValidationExternalObjectsPtr;
 		}
 
-		FDataValidationContext ValidationContext(false, InSettings.ValidationUsecase, ValidationExternalObjects);
+		FDataValidationContext ValidationContext(false, InSettings.ValidationUsecase, ValidationExternalObjects, InSettings.MessageLogName);
 		EDataValidationResult AssetResult = EDataValidationResult::NotValidated;
 		if (!LoadedAsset)
 		{
@@ -1141,7 +1141,7 @@ EDataValidationResult UEditorValidatorSubsystem::ValidateChangelistsInternal(
 
 		FValidateAssetsDetails& Details = OutResults.AssetsDetails.FindOrAdd(Changelist->GetPathName());
 		{
-			FDataValidationContext ValidationContext(false, Settings.ValidationUsecase, {}); // No associated objects for changelist
+			FDataValidationContext ValidationContext(false, Settings.ValidationUsecase, {}, Settings.MessageLogName); // No associated objects for changelist
 			Details.Result = IsObjectValidWithContext(Changelist, ValidationContext);
 			UE::DataValidation::AddAssetValidationMessages(DataValidationLog, ValidationContext);
 			ValidationContext.SplitIssues(Details.ValidationWarnings, Details.ValidationErrors);
@@ -1153,7 +1153,7 @@ EDataValidationResult UEditorValidatorSubsystem::ValidateChangelistsInternal(
 	TSet<FAssetData> AssetsToValidate;	
 	for (UDataValidationChangelist* Changelist : Changelists)
 	{
-		FDataValidationContext ValidationContext(false, Settings.ValidationUsecase, {}); // No associated objects for changelist
+		FDataValidationContext ValidationContext(false, Settings.ValidationUsecase, {}, Settings.MessageLogName); // No associated objects for changelist
 		GatherAssetsToValidateFromChangelist(Changelist, Settings, AssetsToValidate, ValidationContext, OutResults);
 		UE::DataValidation::AddAssetValidationMessages(DataValidationLog, ValidationContext);
 		DataValidationLog.Flush();
@@ -1161,7 +1161,7 @@ EDataValidationResult UEditorValidatorSubsystem::ValidateChangelistsInternal(
 
 	// Filter out assets that we don't want to validate
 	{
-		FDataValidationContext ValidationContext(false, Settings.ValidationUsecase, {}); 
+		FDataValidationContext ValidationContext(false, Settings.ValidationUsecase, {}, Settings.MessageLogName);
 		for (auto It = AssetsToValidate.CreateIterator(); It; ++It)
 		{
 			if (!ShouldValidateAsset(*It, Settings, ValidationContext))
diff --git a/Engine/Plugins/Editor/DataValidation/Source/DataValidation/Public/EditorValidatorHelpers.h b/Engine/Plugins/Editor/DataValidation/Source/DataValidation/Public/EditorValidatorHelpers.h
--- a/Engine/Plugins/Editor/DataValidation/Source/DataValidation/Public/EditorValidatorHelpers.h
+++ b/Engine/Plugins/Editor/DataValidation/Source/DataValidation/Public/EditorValidatorHelpers.h
@@ -22,9 +22,7 @@ namespace EMessageSeverity
 }
 
 namespace UE::DataValidation
-{
-    extern DATAVALIDATION_API const FName MessageLogName; 
-    
+{    
     // Utility to capture log messages from all threads during data validation so they can be filtered
     // and passed as validation errors if desired.
     struct FLogMessageGathererImpl;
diff --git a/Engine/Source/Editor/UnrealEd/Private/CookOnTheFlyServer.cpp b/Engine/Source/Editor/UnrealEd/Private/CookOnTheFlyServer.cpp
--- a/Engine/Source/Editor/UnrealEd/Private/CookOnTheFlyServer.cpp
+++ b/Engine/Source/Editor/UnrealEd/Private/CookOnTheFlyServer.cpp
@@ -3409,9 +3409,7 @@ EDataValidationResult UCookOnTheFlyServer::ValidateSourcePackage(UE::Cook::FPack
 			AssetRegistry->GetAssetsByPath(*ExternalActorsPathForWorld, ExternalObjects, /*bRecursive*/true, /*bIncludeOnlyOnDiskAssets*/true);
 		}
 
-		static const FName NAME_AssetCheck = "AssetCheck";
-
-		FMessageLogScopedOverride AssetCheckLogOverride(NAME_AssetCheck);
+		FMessageLogScopedOverride AssetCheckLogOverride(UE::DataValidation::MessageLogName);
 		if (bLogErrorsAsWarnings)
 		{
 			AssetCheckLogOverride.RemapMessageSeverity(EMessageSeverity::Error, EMessageSeverity::Warning);
diff --git a/Engine/Source/Runtime/CoreUObject/Private/Misc/DataValidation.cpp b/Engine/Source/Runtime/CoreUObject/Private/Misc/DataValidation.cpp
--- a/Engine/Source/Runtime/CoreUObject/Private/Misc/DataValidation.cpp
+++ b/Engine/Source/Runtime/CoreUObject/Private/Misc/DataValidation.cpp
@@ -8,6 +8,11 @@
 
 #include UE_INLINE_GENERATED_CPP_BY_NAME(DataValidation)
 
+namespace UE::DataValidation
+{
+	const FName MessageLogName("AssetCheck");
+}
+
 EDataValidationResult CombineDataValidationResults(EDataValidationResult Result1, EDataValidationResult Result2)
 {
 	/**
diff --git a/Engine/Source/Runtime/CoreUObject/Public/Misc/DataValidation.h b/Engine/Source/Runtime/CoreUObject/Public/Misc/DataValidation.h
--- a/Engine/Source/Runtime/CoreUObject/Public/Misc/DataValidation.h
+++ b/Engine/Source/Runtime/CoreUObject/Public/Misc/DataValidation.h
@@ -9,6 +9,11 @@
 
 #include "DataValidation.generated.h"
 
+namespace UE::DataValidation
+{
+	extern COREUOBJECT_API const FName MessageLogName;
+}
+
 struct FAssetData;
 
 UENUM(BlueprintType)
@@ -59,17 +64,20 @@ public:
 	
 	FDataValidationContext() = default;
 	
-	/** 
+	/**
 	 * Constructor for use by UEditorValidatorSubsystem.
 	 * @param Inusecase allows validators to skip certain validation steps depending on the validation trigger condition
+	 * @param The name of the message log to use for warnings/errors/etc
 	 * @param InAssociatedObjects allows validators to check additional associated objects of the asset being validated e.g. validating external actors
 	 */
 	FDataValidationContext(
-		bool InWasAssetLoadedForValidation,
-		EDataValidationUsecase InUsecase,
-		TConstArrayView<FAssetData> InAssociatedObjects)
+			bool InWasAssetLoadedForValidation,
+			EDataValidationUsecase InUsecase,
+			TConstArrayView<FAssetData> InAssociatedObjects,
+			FName InMessageLogName = UE::DataValidation::MessageLogName)
 		: AssociatedExternalObjects(InAssociatedObjects)
 		, ValidationUsecase(InUsecase)
+		, MessageLogName(InMessageLogName)
 		, bWasAssetLoadedForValidation(InWasAssetLoadedForValidation)
 	{}
 
@@ -114,6 +122,8 @@ public:
 	uint32 GetNumErrors() const { return NumErrors; }
 
 	COREUOBJECT_API void SplitIssues(TArray<FText>& Warnings, TArray<FText>& Errors, TArray<TSharedRef<FTokenizedMessage>>* Messages = nullptr) const;
+
+	FName GetMessageLogName() const { return MessageLogName; }
 	
 private:
 	TArray<FIssue> Issues;
@@ -122,5 +132,6 @@ private:
 	
 	TConstArrayView<FAssetData> AssociatedExternalObjects;
 	EDataValidationUsecase ValidationUsecase = EDataValidationUsecase::None;
+	FName MessageLogName;
 	bool bWasAssetLoadedForValidation = false;
 };
\ No newline at end of file
diff --git a/Engine/Source/Runtime/Engine/Private/DataTable.cpp b/Engine/Source/Runtime/Engine/Private/DataTable.cpp
--- a/Engine/Source/Runtime/Engine/Private/DataTable.cpp
+++ b/Engine/Source/Runtime/Engine/Private/DataTable.cpp
@@ -428,7 +428,7 @@ EDataValidationResult UDataTable::IsDataValid(FDataValidationContext& Context) c
 
 	if (RowStruct && RowStruct->IsChildOf(FTableRowBase::StaticStruct()))
 	{
-		FDataValidationContext RowContext(Context.WasAssetLoadedForValidation(), Context.GetValidationUsecase(), Context.GetAssociatedExternalObjects());
+		FDataValidationContext RowContext(Context.WasAssetLoadedForValidation(), Context.GetValidationUsecase(), Context.GetAssociatedExternalObjects(), Context.GetMessageLogName());
 		int32 LastNumIssues = 0;
 		for (const TPair<FName, uint8*>& TableRowPair : RowMap)
 		{
