PR #14157: Fix for data validation fixers that look for assets to be marked as dirty.

#Synced-CL 49593019

diff --git a/Engine/Plugins/Editor/DataValidation/Source/DataValidation/Private/DataValidationFixers.cpp b/Engine/Plugins/Editor/DataValidation/Source/DataValidation/Private/DataValidationFixers.cpp
--- a/Engine/Plugins/Editor/DataValidation/Source/DataValidation/Private/DataValidationFixers.cpp
+++ b/Engine/Plugins/Editor/DataValidation/Source/DataValidation/Private/DataValidationFixers.cpp
@@ -74,9 +74,9 @@ auto CollectDirtyPackagesDuring(TSet<UPackage*>& OutPackagesMarkedDirty, F&& Lam
 {
 	FDelegateHandle PackageModificationListenerHandle = UPackage::PackageMarkedDirtyEvent.AddLambda(
 		// Lambda capture safety: This lambda is removed from the delegate before PackagesMarkedDirty goes out of scope.
-		[&OutPackagesMarkedDirty](UPackage* Package, bool bIsDirty)
+		[&OutPackagesMarkedDirty](UPackage* Package, bool bWasDirty)
 		{
-			if (bIsDirty)
+			if (Package)
 			{
 				OutPackagesMarkedDirty.Add(Package);
 			}
