Raise maximum value for DisableValidateOnSaveCount from 255 to 4 billion
#rnx
#Synced-CL 49673251

diff --git a/Engine/Plugins/Editor/DataValidation/Source/DataValidation/Public/EditorValidatorSubsystem.h b/Engine/Plugins/Editor/DataValidation/Source/DataValidation/Public/EditorValidatorSubsystem.h
--- a/Engine/Plugins/Editor/DataValidation/Source/DataValidation/Public/EditorValidatorSubsystem.h
+++ b/Engine/Plugins/Editor/DataValidation/Source/DataValidation/Public/EditorValidatorSubsystem.h
@@ -566,7 +566,7 @@ protected:
 	/**
 	 * Counter used by PushDisableValidateOnSave and PopDisableValidateOnSave to know whether "validate on save" is temporarily disabled.
 	 */
-	uint8 DisableValidateOnSaveCount = 0;
+	uint32 DisableValidateOnSaveCount = 0;
 };
 
 #undef UE_API
