Cook-time validation improvements

* Pre-load external objects prior to running validators. This is controlled by the CVar 'DataValidation.LoadExternalObjectsForCookValidation' and is enabled by default to match the default editor behavior.
* Added support for reporting the total validation duration per-validator for cook-time validation. This is controlled by the CVar 'DataValidation.ReportCookValidationStats' and is enabled by default.

#Synced-CL 49914824

diff --git a/Engine/Plugins/Editor/DataValidation/Source/DataValidation/Private/DataValidationModule.cpp b/Engine/Plugins/Editor/DataValidation/Source/DataValidation/Private/DataValidationModule.cpp
--- a/Engine/Plugins/Editor/DataValidation/Source/DataValidation/Private/DataValidationModule.cpp
+++ b/Engine/Plugins/Editor/DataValidation/Source/DataValidation/Private/DataValidationModule.cpp
@@ -15,6 +15,8 @@
 #include "EditorValidatorBase.h"
 #include "EditorValidator_Material.h"
 #include "Elements/Framework/TypedElementSelectionSet.h"
+#include "HAL/IConsoleManager.h"
+#include "LoadExternalObjectsForValidation.h"
 #include "Logging/MessageLog.h"
 #include "UObject/ICookInfo.h"
 
@@ -32,6 +34,21 @@
 
 #define LOCTEXT_NAMESPACE "DataValidationModule"
 
+namespace DataValidationOptions
+{
+	static bool bLoadExternalObjectsForCookValidation = true;
+	static FAutoConsoleVariableRef CVarLoadExternalObjectsForCookValidation(
+		TEXT("DataValidation.LoadExternalObjectsForCookValidation"),
+		bLoadExternalObjectsForCookValidation,
+		TEXT("True to load any external objects (eg, actors) prior to running cook-time validation on the asset"));
+
+	static bool bReportCookValidationStats = true;
+	static FAutoConsoleVariableRef CVarReportCookValidationStats(
+		TEXT("DataValidation.ReportCookValidationStats"),
+		bReportCookValidationStats,
+		TEXT("True to report post-cook statistics for each validator"));
+}
+
 class FDataValidationModule : public IDataValidationModule
 {
 	// Begin IModuleInterface
@@ -47,6 +64,8 @@ class FDataValidationModule : public IDataValidationModule
 private:
 	void OnPackageSaved(const FString& PackageFileName, UPackage* Package, FObjectPostSaveContext ObjectSaveContext);
 
+	void OnCookStarted(UE::Cook::ICookInfo& CookInfo);
+	void OnCookFinished(UE::Cook::ICookInfo& CookInfo);
 	EDataValidationResult OnValidateSourcePackageDuringCook(UPackage* Package, FDataValidationContext& ValidationContext);
 
 	// Adds Asset and any assets it depends on to the set DependentAssets
@@ -55,12 +74,16 @@ private:
 	void RegisterMenus();
 	static FText Menu_ValidateDataGetTitle();
 	static void Menu_ValidateData();
+
+	TMap<FTopLevelAssetPath, FValidatorStatistics> CookValidationStats;
 };
 
 IMPLEMENT_MODULE(FDataValidationModule, DataValidation)
 
 void FDataValidationModule::StartupModule()
 {	
+	UE::Cook::FDelegates::CookStarted.AddRaw(this, &FDataValidationModule::OnCookStarted);
+	UE::Cook::FDelegates::CookFinished.AddRaw(this, &FDataValidationModule::OnCookFinished);
 	UE::Cook::FDelegates::ValidateSourcePackage.BindRaw(this, &FDataValidationModule::OnValidateSourcePackageDuringCook);
 
 	if (!IsRunningCommandlet() && !IsRunningGame() && FSlateApplication::IsInitialized())
@@ -84,6 +107,8 @@ void FDataValidationModule::StartupModule()
 
 void FDataValidationModule::ShutdownModule()
 {
+	UE::Cook::FDelegates::CookStarted.RemoveAll(this);
+	UE::Cook::FDelegates::CookFinished.RemoveAll(this);
 	UE::Cook::FDelegates::ValidateSourcePackage.Unbind();
 
 	if (!IsRunningCommandlet() && !IsRunningGame() && !IsRunningDedicatedServer())
@@ -381,26 +406,59 @@ void FDataValidationModule::OnPackageSaved(const FString& PackageFileName, UPack
 	}
 }
 
-EDataValidationResult FDataValidationModule::OnValidateSourcePackageDuringCook(UPackage* Package, FDataValidationContext& ValidationContext)
+void FDataValidationModule::OnCookStarted(UE::Cook::ICookInfo& CookInfo)
 {
-	EDataValidationResult FinalValidationResult = EDataValidationResult::NotValidated;
+	using namespace UE::Cook;
 
-	if (UEditorValidatorSubsystem* EditorValidationSubsystem = GEditor->GetEditorSubsystem<UEditorValidatorSubsystem>())
+	if (EnumHasAnyFlags(CookInfo.GetCookValidationOptions(), ECookValidationOptions::RunAssetValidation))
 	{
-		// Log the enabled set of validators, but only once
+		// Log the enabled validators
+		// This will also preload any BP-based validators before we start the main cook
+		if (UEditorValidatorSubsystem* EditorValidationSubsystem = GEditor->GetEditorSubsystem<UEditorValidatorSubsystem>())
 		{
-			static const bool LogValidatorsListOnce = [EditorValidationSubsystem]()
+			int32 NumEnabledValidators = 0;
+
+			UE_LOG(LogContentValidation, Display, TEXT("Enabled validators:"));
+			EditorValidationSubsystem->ForEachEnabledValidator([&NumEnabledValidators](UEditorValidatorBase* Validator)
 			{
-				UE_LOG(LogContentValidation, Display, TEXT("Enabled validators:"));
-				EditorValidationSubsystem->ForEachEnabledValidator([](UEditorValidatorBase* Validator)
-				{
-					UE_LOG(LogContentValidation, Display, TEXT("\t%s"), *Validator->GetClass()->GetClassPathName().ToString());
-					return true;
-				});
+				++NumEnabledValidators;
+				UE_LOG(LogContentValidation, Display, TEXT("\t%s"), *Validator->GetClass()->GetClassPathName().ToString());
 				return true;
-			}();
+			});
+
+			CookValidationStats.Reserve(NumEnabledValidators);
+		}
+	}
+}
+
+void FDataValidationModule::OnCookFinished(UE::Cook::ICookInfo& CookInfo)
+{
+	using namespace UE::Cook;
+
+	if (DataValidationOptions::bReportCookValidationStats && EnumHasAnyFlags(CookInfo.GetCookValidationOptions(), ECookValidationOptions::RunAssetValidation))
+	{
+		// Report the validators, slowest to fastest
+		CookValidationStats.ValueSort([](const FValidatorStatistics& One, const FValidatorStatistics& Two)
+			{
+				return One.TotalDurationSeconds > Two.TotalDurationSeconds;
+			});
+
+		UE_LOG(LogContentValidation, Display, TEXT("Validation stats:"));
+		for (const TTuple<FTopLevelAssetPath, FValidatorStatistics> CookValidationStatsPair : CookValidationStats)
+		{
+			UE_LOG(LogContentValidation, Display, TEXT("\t%s validated %d assets in %f seconds"), *CookValidationStatsPair.Key.ToString(), CookValidationStatsPair.Value.AssetsValidated, CookValidationStatsPair.Value.TotalDurationSeconds);
 		}
+	}
 
+	CookValidationStats.Reset();
+}
+
+EDataValidationResult FDataValidationModule::OnValidateSourcePackageDuringCook(UPackage* Package, FDataValidationContext& ValidationContext)
+{
+	EDataValidationResult FinalValidationResult = EDataValidationResult::NotValidated;
+
+	if (UEditorValidatorSubsystem* EditorValidationSubsystem = GEditor->GetEditorSubsystem<UEditorValidatorSubsystem>())
+	{
 		TArray<FAssetData> AssetList;
 		IAssetRegistry::GetChecked().GetAssetsByPackageName(Package->GetFName(), AssetList);
 
@@ -431,8 +489,10 @@ EDataValidationResult FDataValidationModule::OnValidateSourcePackageDuringCook(U
 						->AddToken(FAssetDataToken::Create(AssetData))
 						->AddToken(FTextToken::Create(LOCTEXT("Data.ValidatingAsset", "Validating asset")));
 
+					UE::DataValidation::FScopedLoadExternalObjects ExternalObjectsLoader(Asset, ValidationContext, DataValidationOptions::bLoadExternalObjectsForCookValidation);
+
 					UE::DataValidation::FScopedLogMessageGatherer LogGatherer;
-					EDataValidationResult ValidationResult = EditorValidationSubsystem->IsObjectValidWithContext(Asset, ValidationContext);
+					EDataValidationResult ValidationResult = EditorValidationSubsystem->ValidateObjectInternal(AssetData, Asset, ValidationContext, DataValidationOptions::bReportCookValidationStats ? &CookValidationStats : nullptr);
 
 					TArray<FString> LogWarnings;
 					TArray<FString> LogErrors;
diff --git a/Engine/Plugins/Editor/DataValidation/Source/DataValidation/Private/EditorValidatorSubsystem.cpp b/Engine/Plugins/Editor/DataValidation/Source/DataValidation/Private/EditorValidatorSubsystem.cpp
--- a/Engine/Plugins/Editor/DataValidation/Source/DataValidation/Private/EditorValidatorSubsystem.cpp
+++ b/Engine/Plugins/Editor/DataValidation/Source/DataValidation/Private/EditorValidatorSubsystem.cpp
@@ -33,6 +33,7 @@
 #include "Misc/ScopedSlowTask.h"
 #include "Misc/ScopeExit.h"
 #include "PackageTools.h"
+#include "ProfilingDebugging/ScopedTimers.h"
 #include "TelemetryRouter.h"
 #include "UObject/ICookInfo.h"
 #include "UObject/TopLevelAssetPath.h"
@@ -449,12 +450,29 @@ EDataValidationResult UEditorValidatorSubsystem::ValidateObjectInternal(
 		{
 			UE_LOG(LogContentValidation, Verbose, TEXT("Validating '%s' with '%s'..."), *InObject->GetPathName(), *Validator->GetClass()->GetName());
 			GInitRunaway(); // Reset runaway counter, as ValidateLoadedAsset may be implemented in a BP and could overflow the runaway count due to being called in a loop
-			EDataValidationResult NewResult = Validator->ValidateLoadedAsset(InAssetData, InObject, InContext);
-			if (NewResult != EDataValidationResult::NotValidated && ValidatorStatistics)
+
+			EDataValidationResult NewResult = EDataValidationResult::NotValidated;
+			if (ValidatorStatistics)
+			{
+				double ValidationDuration = 0.0;
+				{
+					FScopedDurationTimer ValidationTimer(ValidationDuration);
+					NewResult = Validator->ValidateLoadedAsset(InAssetData, InObject, InContext);
+				}
+
+				FValidatorStatistics& ThisValidatorStats = ValidatorStatistics->FindOrAdd(Validator->GetClass()->GetClassPathName());
+				if (NewResult != EDataValidationResult::NotValidated)
+				{
+					ThisValidatorStats.AssetsValidated++;
+				}
+				ThisValidatorStats.TotalDurationSeconds += ValidationDuration;
+			}
+			else
 			{
-				ValidatorStatistics->FindOrAdd(FTopLevelAssetPath(Validator->GetClass())).AssetsValidated++;
+				NewResult = Validator->ValidateLoadedAsset(InAssetData, InObject, InContext);
 			}
 			Result = CombineDataValidationResults(Result, NewResult);
+
 			return true;
 		});
 	}
diff --git a/Engine/Plugins/Editor/DataValidation/Source/DataValidation/Public/EditorValidatorSubsystem.h b/Engine/Plugins/Editor/DataValidation/Source/DataValidation/Public/EditorValidatorSubsystem.h
--- a/Engine/Plugins/Editor/DataValidation/Source/DataValidation/Public/EditorValidatorSubsystem.h
+++ b/Engine/Plugins/Editor/DataValidation/Source/DataValidation/Public/EditorValidatorSubsystem.h
@@ -82,6 +82,10 @@ struct FValidatorStatistics
 	/** Number of assets that this validator marked as needing validation e.g. because a dependency changed */
 	UPROPERTY(VisibleAnywhere, BlueprintReadOnly, Category="Validation")
 	int32 AssetsAddedForValidation = 0;
+
+	/** Total time (in seconds) that this validator spent validating assets */
+	UPROPERTY(VisibleAnywhere, BlueprintReadOnly, Category = "Validation")
+	double TotalDurationSeconds = 0.0;
 };
 
 USTRUCT(BlueprintType)
@@ -488,6 +492,7 @@ private:
 	void RemoveValidator_Internal(UEditorValidatorBase* InValidator);
 	void RemoveValidator_Internal(FTopLevelAssetPath InValidatorClass);
 protected:
+	friend class FDataValidationModule;
 
 	/** 
 	 * Validate a set of assets, adding the results to the log/output of a higher level task (e.g. changelist validation).
