Add cvar to hide ambiguous Scripted Asset Actions (SupportedClasses = empty).

Additionally, add data validation to UAssetActionUtility and UActorActionUtility to throw an error when SupportedClasses is empty.

#Synced-CL 49762516

diff --git a/Engine/Plugins/Editor/DataValidation/Source/DataValidation/Private/EditorValidator_ActionUtility.cpp b/Engine/Plugins/Editor/DataValidation/Source/DataValidation/Private/EditorValidator_ActionUtility.cpp
new file mode 100644
--- /dev/null
+++ b/Engine/Plugins/Editor/DataValidation/Source/DataValidation/Private/EditorValidator_ActionUtility.cpp
@@ -0,0 +1,67 @@
+// Copyright Epic Games, Inc. All Rights Reserved.
+
+#include "EditorValidator_ActionUtility.h"
+
+#include "ActorActionUtility.h"
+#include "AssetActionUtility.h"
+#include "Engine/Blueprint.h"
+#include "EditorValidatorSubsystem.h"
+
+#include UE_INLINE_GENERATED_CPP_BY_NAME(EditorValidator_ActionUtility)
+
+#define LOCTEXT_NAMESPACE "ActionUtilityValidator"
+
+namespace EditorValidatorActionUtilityHelpers
+{
+	UObject* GetObjectFromAssetPtr(UObject* InAsset)
+	{
+		UObject* Object = InAsset;
+		if (const UBlueprint* BP = Cast<UBlueprint>(InAsset))
+		{
+			Object = BP->GeneratedClass.GetDefaultObject();
+		}
+		
+		return Object;
+	}
+};
+
+bool UEditorValidator_ActionUtility::CanValidateAsset_Implementation(const FAssetData& AssetData, UObject* InAsset, FDataValidationContext& InContext) const
+{
+	check(InAsset);
+	UObject* Object = EditorValidatorActionUtilityHelpers::GetObjectFromAssetPtr(InAsset);
+	return Cast<UAssetActionUtility>(Object) || Cast<UActorActionUtility>(Object);
+}
+
+EDataValidationResult UEditorValidator_ActionUtility::ValidateLoadedAsset_Implementation(const FAssetData& AssetData, UObject* InAsset, FDataValidationContext& InContext)
+{
+	check(InAsset);
+	EDataValidationResult Result = EDataValidationResult::Valid;
+	UObject* Object = EditorValidatorActionUtilityHelpers::GetObjectFromAssetPtr(InAsset);
+	
+	if (UActorActionUtility* ActorActionUtility = Cast<UActorActionUtility>(Object))
+	{
+		if (ActorActionUtility->GetSupportedClasses().IsEmpty())
+		{
+			AssetFails(Object, LOCTEXT("ActorActionUtility_NoSupportedClasses", 
+				"Supported Classes array must not be empty. "
+				"Add only the class(es) that your action supports. "
+				"If your action is truly compatible with all types of Actors, add the 'Actor' class explicitly"));
+			Result = EDataValidationResult::Invalid;
+		}
+	}
+	else if (UAssetActionUtility* AssetActionUtility = Cast<UAssetActionUtility>(Object))
+	{
+		if (AssetActionUtility->GetSupportedClasses().IsEmpty())
+		{
+			AssetFails(Object, LOCTEXT("AssetActionUtility_NoSupportedClasses", 
+				"Supported Classes array must not be empty. "
+				"Add only the class(es) that your action supports. "
+				"If your action is truly compatible with all types of Assets, add the 'Object' class explicitly"));
+			Result = EDataValidationResult::Invalid;
+		}
+	}
+	
+	return Result;
+}
+
+#undef LOCTEXT_NAMESPACE
diff --git a/Engine/Plugins/Editor/DataValidation/Source/DataValidation/Public/EditorValidator_ActionUtility.h b/Engine/Plugins/Editor/DataValidation/Source/DataValidation/Public/EditorValidator_ActionUtility.h
new file mode 100644
--- /dev/null
+++ b/Engine/Plugins/Editor/DataValidation/Source/DataValidation/Public/EditorValidator_ActionUtility.h
@@ -0,0 +1,22 @@
+﻿// Copyright Epic Games, Inc. All Rights Reserved.
+
+#pragma once
+
+#include "EditorValidatorBase.h"
+
+#include "EditorValidator_ActionUtility.generated.h"
+
+class UObject;
+class FDataValidationContext; // see UObject.h
+struct FAssetData;
+
+// Checks if UActorActionUtility and UAssetActionUtility data is valid
+UCLASS()
+class UEditorValidator_ActionUtility : public UEditorValidatorBase
+{
+	GENERATED_BODY()
+
+protected:
+	virtual bool CanValidateAsset_Implementation(const FAssetData& InAssetData, UObject* InAsset, FDataValidationContext& InContext) const override;
+	virtual EDataValidationResult ValidateLoadedAsset_Implementation(const FAssetData& InAssetData, UObject* InAsset, FDataValidationContext& InContext) override;
+};
diff --git a/Engine/Source/Editor/Blutility/Private/EditorUtilityAssetPrototype.cpp b/Engine/Source/Editor/Blutility/Private/EditorUtilityAssetPrototype.cpp
--- a/Engine/Source/Editor/Blutility/Private/EditorUtilityAssetPrototype.cpp
+++ b/Engine/Source/Editor/Blutility/Private/EditorUtilityAssetPrototype.cpp
@@ -7,11 +7,24 @@
 #include "Engine/BlueprintGeneratedClass.h"
 #include "JsonObjectConverter.h"
 #include "AssetRegistry/AssetData.h"
+#include "Editor/EditorPerProjectUserSettings.h"
+#include "HAL/IConsoleManager.h"
 #include "Serialization/JsonSerializerMacros.h"
 #include "UObject/AssetRegistryTagsContext.h"
 
 #include UE_INLINE_GENERATED_CPP_BY_NAME(EditorUtilityAssetPrototype)
 
+namespace AssetActionUtilityConsoleVariables
+{
+	bool bHideAmbiguousScriptedAssetActions = true;
+	FAutoConsoleVariableRef CVarEnableMutatorsSubsystem(
+		TEXT("AssetActionUtility.HideAmbiguousScriptedAssetActions"),
+		bHideAmbiguousScriptedAssetActions,
+		TEXT("If enabled, Scripted Asset Actions will only show up if they explicitly define Supported Class(es) via their Class Defaults.\n0: Disable, 1: Enable (default)"),
+		ECVF_Default);
+
+} // namespace AssetActionUtilityConsoleVariables
+
 const FName AssetActionUtilityTags::BlutilityTagVersion(TEXT("BlutilityTagVersion"));
 const FName AssetActionUtilityTags::SupportedClasses(TEXT("SupportedClasses"));
 const FName AssetActionUtilityTags::IsActionForBlueprint(TEXT("IsActionForBlueprint"));
@@ -75,7 +88,12 @@ TArray<TSoftClassPtr<UObject>> FAssetActionUtilityPrototype::GetSupportedClasses
 
 	if (SupportedClasses.Num() == 0)
 	{
-		SupportedClasses.Add(TSoftClassPtr<UObject>(UObject::StaticClass()));
+		// If this cvar is set, we DON'T assume an empty array means we support UObject and derived classes
+		// This cleans up the "Scripted Asset Actions" list and ensures only actions that explicitly support the asset class being requested show up
+		if (!AssetActionUtilityConsoleVariables::bHideAmbiguousScriptedAssetActions)
+		{
+			SupportedClasses.Add(TSoftClassPtr<UObject>(UObject::StaticClass()));
+		}
 	}
 	
 	return SupportedClasses;
