[DataValidation] Fix crash that could occur when unloading actors with child components during cook validation
Use a deferred loading context to ensure that all actors are added or removed in a single operation, preventing improper handling of actors with child actor components.

#Synced-CL 50059267

diff --git a/Engine/Plugins/Editor/DataValidation/Source/DataValidation/Private/LoadExternalObjectsForValidation.cpp b/Engine/Plugins/Editor/DataValidation/Source/DataValidation/Private/LoadExternalObjectsForValidation.cpp
--- a/Engine/Plugins/Editor/DataValidation/Source/DataValidation/Private/LoadExternalObjectsForValidation.cpp
+++ b/Engine/Plugins/Editor/DataValidation/Source/DataValidation/Private/LoadExternalObjectsForValidation.cpp
@@ -101,20 +101,27 @@ FScopedLoadExternalObjects::FScopedLoadExternalObjects(UObject* InAsset, FDataVa
             }
         }
         
-        for (const FAssetData& ActorData : ActorsData)
-        {
-            if (const FWorldPartitionActorDescInstance* ActorDesc = ContainersToValidate.GetActorDescInstanceByPath(ActorData.AssetName.ToString()))
-            {
-                ActorRefs.Add(FWorldPartitionReference(&ContainersToValidate, ActorDesc->GetGuid()));
-            }
-        }
+		{
+			FWorldPartitionLoadingContext::FDeferred LoadingContext;
+			for (const FAssetData& ActorData : ActorsData)
+			{
+				if (const FWorldPartitionActorDescInstance* ActorDesc = ContainersToValidate.GetActorDescInstanceByPath(ActorData.AssetName.ToString()))
+				{
+					ActorRefs.Add(FWorldPartitionReference(&ContainersToValidate, ActorDesc->GetGuid()));
+				}
+			}
+		}
     }
 }
 
 FScopedLoadExternalObjects::~FScopedLoadExternalObjects()
 {
     // Explicitly release actors before uninitializing containers
-    ActorRefs.Empty();
+	{
+		FWorldPartitionLoadingContext::FDeferred LoadingContext;
+		ActorRefs.Empty();
+	}
+
 	for (const TStrongObjectPtr<UActorDescContainerInstance>& ContainerToUninit : ContainersToUninit)
 	{
 		ContainerToUninit->Uninitialize();
