Added the cook-time validation stats to the cook stats analytic event

#rnx
#Synced-CL 50687180

diff --git a/Engine/Plugins/Editor/DataValidation/Source/DataValidation/Private/DataValidationModule.cpp b/Engine/Plugins/Editor/DataValidation/Source/DataValidation/Private/DataValidationModule.cpp
--- a/Engine/Plugins/Editor/DataValidation/Source/DataValidation/Private/DataValidationModule.cpp
+++ b/Engine/Plugins/Editor/DataValidation/Source/DataValidation/Private/DataValidationModule.cpp
@@ -66,6 +66,9 @@ private:
 
 	void OnCookStarted(UE::Cook::ICookInfo& CookInfo);
 	void OnCookFinished(UE::Cook::ICookInfo& CookInfo);
+#if ENABLE_COOK_STATS
+	void LogCookStats(FCookStatsManager::AddStatFuncRef AddStat);
+#endif
 	EDataValidationResult OnValidateSourcePackageDuringCook(UPackage* Package, FDataValidationContext& ValidationContext);
 
 	// Adds Asset and any assets it depends on to the set DependentAssets
@@ -84,6 +87,9 @@ void FDataValidationModule::StartupModule()
 {	
 	UE::Cook::FDelegates::CookStarted.AddRaw(this, &FDataValidationModule::OnCookStarted);
 	UE::Cook::FDelegates::CookFinished.AddRaw(this, &FDataValidationModule::OnCookFinished);
+#if ENABLE_COOK_STATS
+	FCookStatsManager::CookStatsCallbacks.AddRaw(this, &FDataValidationModule::LogCookStats);
+#endif
 	UE::Cook::FDelegates::ValidateSourcePackage.BindRaw(this, &FDataValidationModule::OnValidateSourcePackageDuringCook);
 
 	if (!IsRunningCommandlet() && !IsRunningGame() && FSlateApplication::IsInitialized())
@@ -109,6 +115,9 @@ void FDataValidationModule::ShutdownModule()
 {
 	UE::Cook::FDelegates::CookStarted.RemoveAll(this);
 	UE::Cook::FDelegates::CookFinished.RemoveAll(this);
+#if ENABLE_COOK_STATS
+	FCookStatsManager::CookStatsCallbacks.RemoveAll(this);
+#endif
 	UE::Cook::FDelegates::ValidateSourcePackage.Unbind();
 
 	if (!IsRunningCommandlet() && !IsRunningGame() && !IsRunningDedicatedServer())
@@ -410,6 +419,8 @@ void FDataValidationModule::OnCookStarted(UE::Cook::ICookInfo& CookInfo)
 {
 	using namespace UE::Cook;
 
+	CookValidationStats.Reset();
+
 	if (EnumHasAnyFlags(CookInfo.GetCookValidationOptions(), ECookValidationOptions::RunAssetValidation))
 	{
 		// Log the enabled validators
@@ -444,14 +455,19 @@ void FDataValidationModule::OnCookFinished(UE::Cook::ICookInfo& CookInfo)
 			});
 
 		UE_LOG(LogContentValidation, Display, TEXT("Validation stats:"));
-		for (const TTuple<FTopLevelAssetPath, FValidatorStatistics> CookValidationStatsPair : CookValidationStats)
+		for (const TTuple<FTopLevelAssetPath, FValidatorStatistics>& CookValidationStatsPair : CookValidationStats)
 		{
 			UE_LOG(LogContentValidation, Display, TEXT("\t%s validated %d assets in %f seconds"), *CookValidationStatsPair.Key.ToString(), CookValidationStatsPair.Value.AssetsValidated, CookValidationStatsPair.Value.TotalDurationSeconds);
 		}
 	}
+}
 
-	CookValidationStats.Reset();
+#if ENABLE_COOK_STATS
+void FDataValidationModule::LogCookStats(FCookStatsManager::AddStatFuncRef AddStat)
+{
+	UE::DataValidation::ValidationStatsToAnalyticEventData(CookValidationStats, AddStat);
 }
+#endif
 
 EDataValidationResult FDataValidationModule::OnValidateSourcePackageDuringCook(UPackage* Package, FDataValidationContext& ValidationContext)
 {
diff --git a/Engine/Plugins/Editor/DataValidation/Source/DataValidation/Private/EditorValidatorHelpers.cpp b/Engine/Plugins/Editor/DataValidation/Source/DataValidation/Private/EditorValidatorHelpers.cpp
--- a/Engine/Plugins/Editor/DataValidation/Source/DataValidation/Private/EditorValidatorHelpers.cpp
+++ b/Engine/Plugins/Editor/DataValidation/Source/DataValidation/Private/EditorValidatorHelpers.cpp
@@ -3,8 +3,10 @@
 #include "EditorValidatorHelpers.h"
 
 #include "Algo/AllOf.h"
+#include "AnalyticsEventAttribute.h"
 #include "AssetRegistry/AssetData.h"
 #include "AssetRegistry/AssetDataToken.h"
+#include "EditorValidatorSubsystem.h"
 #include "Logging/MessageLog.h"
 #include "Containers/ArrayView.h"
 #include "Misc/DataValidation.h"
@@ -91,6 +93,33 @@ FText MakeTimestampedMessageLogPageTitle(const FText& BaseTitle, const FDateTime
 	return FText::Format(TimestampedPageTitleFmt, FText::AsDateTime(UtcTimestamp, EDateTimeStyle::Short, EDateTimeStyle::Default), BaseTitle);
 }
 
+void ValidationStatsToAnalyticEventData(const TMap<FTopLevelAssetPath, FValidatorStatistics>& ValidationStats, TArray<FAnalyticsEventAttribute>& OutEventAttributes)
+{
+	for (const TTuple<FTopLevelAssetPath, FValidatorStatistics>& ValidationStatsPair : ValidationStats)
+	{
+		// This matches the format used by FEditorTelemetry::RecordEvent_Cooking, where StatName is combined with each AttrName to form "{StatName}_{AttrName}"
+		// (see the below variant of ValidationStatsToAnalyticEventData)
+		const FString StatName = FString::Printf(TEXT("Validation_%s"), FNameBuilder(ValidationStatsPair.Key.GetAssetName()).ToString());
+		OutEventAttributes.Emplace(FString::Printf(TEXT("%s_TotalDurationSecs"), *StatName), ValidationStatsPair.Value.TotalDurationSeconds);
+		OutEventAttributes.Emplace(FString::Printf(TEXT("%s_AssetsValidated"), *StatName), ValidationStatsPair.Value.AssetsValidated);
+	}
+}
+
+#if ENABLE_COOK_STATS
+void ValidationStatsToAnalyticEventData(const TMap<FTopLevelAssetPath, FValidatorStatistics>& ValidationStats, FCookStatsManager::AddStatFuncRef AddStat)
+{
+	for (const TTuple<FTopLevelAssetPath, FValidatorStatistics>& ValidationStatsPair : ValidationStats)
+	{
+		// StatName is combined with each AttrName in FEditorTelemetry::RecordEvent_Cooking, to form "{StatName}_{AttrName}"
+		// (see the above variant of ValidationStatsToAnalyticEventData)
+		AddStat(
+			FString::Printf(TEXT("Validation_%s"), FNameBuilder(ValidationStatsPair.Key.GetAssetName()).ToString()),
+			FCookStatsManager::CreateKeyValueArray(TEXT("TotalDurationSecs"), ValidationStatsPair.Value.TotalDurationSeconds, TEXT("AssetsValidated"), ValidationStatsPair.Value.AssetsValidated)
+			);
+	}
+}
+#endif
+
 struct FLogMessageGathererImpl : public FOutputDevice
 {
     FLogMessageGathererImpl()
diff --git a/Engine/Plugins/Editor/DataValidation/Source/DataValidation/Public/EditorValidatorHelpers.h b/Engine/Plugins/Editor/DataValidation/Source/DataValidation/Public/EditorValidatorHelpers.h
--- a/Engine/Plugins/Editor/DataValidation/Source/DataValidation/Public/EditorValidatorHelpers.h
+++ b/Engine/Plugins/Editor/DataValidation/Source/DataValidation/Public/EditorValidatorHelpers.h
@@ -4,6 +4,7 @@
 
 #include "Containers/ContainersFwd.h"
 #include "Delegates/Delegate.h"
+#include "ProfilingDebugging/CookStats.h"
 #include "Templates/Tuple.h"
 #include "UObject/NameTypes.h"
 
@@ -11,11 +12,15 @@
 
 class FMessageLog;
 class FText;
-struct FAssetData;
 class FDataValidationContext;
 class FWildcardString;
 class UObject;
 
+struct FAnalyticsEventAttribute;
+struct FAssetData;
+struct FTopLevelAssetPath;
+struct FValidatorStatistics;
+
 namespace EMessageSeverity
 {
     enum Type : int;
@@ -118,6 +123,11 @@ namespace UE::DataValidation
 	/** Prepend the given UTC timestamp (or UtcNow if none provided) to the given base title to provide a pseudo-unique title, eg) "[Timestamp] BaseTitle" */
 	UE_API FText MakeTimestampedMessageLogPageTitle(const FText& BaseTitle);
 	UE_API FText MakeTimestampedMessageLogPageTitle(const FText& BaseTitle, const FDateTime& UtcTimestamp);
+
+	UE_API void ValidationStatsToAnalyticEventData(const TMap<FTopLevelAssetPath, FValidatorStatistics>& ValidationStats, TArray<FAnalyticsEventAttribute>& OutEventAttributes);
+#if ENABLE_COOK_STATS
+	UE_API void ValidationStatsToAnalyticEventData(const TMap<FTopLevelAssetPath, FValidatorStatistics>& ValidationStats, FCookStatsManager::AddStatFuncRef AddStat);
+#endif
 }
 
 #undef UE_API
