Inline init, and make sure a transient property is marked as such

#rnx
#Synced-CL 50847432

diff --git a/Engine/Plugins/Editor/DataValidation/Source/DataValidation/Private/EditorValidatorBase.cpp b/Engine/Plugins/Editor/DataValidation/Source/DataValidation/Private/EditorValidatorBase.cpp
--- a/Engine/Plugins/Editor/DataValidation/Source/DataValidation/Private/EditorValidatorBase.cpp
+++ b/Engine/Plugins/Editor/DataValidation/Source/DataValidation/Private/EditorValidatorBase.cpp
@@ -13,12 +13,6 @@
 
 #define LOCTEXT_NAMESPACE "AssetValidation"
 
-UEditorValidatorBase::UEditorValidatorBase()
-	: Super()
-{
-	bIsEnabled = true;
-}
-
 EDataValidationResult UEditorValidatorBase::ValidateLoadedAsset(const FAssetData& AssetData, UObject* Asset, FDataValidationContext& Context)
 {
 	EDataValidationResult Result = EDataValidationResult::NotValidated;
diff --git a/Engine/Plugins/Editor/DataValidation/Source/DataValidation/Public/EditorValidatorBase.h b/Engine/Plugins/Editor/DataValidation/Source/DataValidation/Public/EditorValidatorBase.h
--- a/Engine/Plugins/Editor/DataValidation/Source/DataValidation/Public/EditorValidatorBase.h
+++ b/Engine/Plugins/Editor/DataValidation/Source/DataValidation/Public/EditorValidatorBase.h
@@ -30,8 +30,6 @@ class UEditorValidatorBase : public UObject
 	GENERATED_BODY()
 
 public:
-	UE_API UEditorValidatorBase();
-	
 	// Validation entry point - will dispatch to native and/or BP validation as appropriate and return results.
 	// If validator doesn't want to validate this asset/in this context, it will return NotValidated which is acceptable.
 	UE_API EDataValidationResult ValidateLoadedAsset(const FAssetData& InAssetData, UObject* InAsset, FDataValidationContext& Context);
@@ -138,27 +136,25 @@ public:
 
 protected:
 	UPROPERTY(EditAnywhere, Category = "Asset Validation", meta = (BlueprintProtected = "true"))
-	bool bIsEnabled;
-	
-	// Allows disabling validators from config without checking in assets/modifying code.
+	bool bIsEnabled = true;
+
+	/* Allows disabling validators from config without checking in assets/modifying code. */
 	UPROPERTY(Config)
-	bool bIsConfigDisabled;
+	bool bIsConfigDisabled = false;
 
-	/* Whether we should also print out the source validator when printing validation errors.*/
+	/* Whether we should also print out the source validator when printing validation errors. */
 	UPROPERTY(EditAnywhere, Category = "Asset Validation", meta = (BlueprintProtected = "true"))
 	bool bOnlyPrintCustomMessage = false;
 
 private:
-	EDataValidationResult ValidationResult;
+	EDataValidationResult ValidationResult = EDataValidationResult::NotValidated;
 	TArray<FText> AllWarnings;
 	TArray<FText> AllErrors;
 	TArray<TSharedRef<FTokenizedMessage>> AllMessages;
 
-	
-	UPROPERTY()
-	TObjectPtr<UObject> CurrentObjectBeingValidated;
-
-	const FAssetData* CurrentAssetBeingValidated;
+	UPROPERTY(Transient)
+	TObjectPtr<UObject> CurrentObjectBeingValidated = nullptr;
+	const FAssetData* CurrentAssetBeingValidated = nullptr;
 };
 
 #undef UE_API
