Made internal data validation API protected

#rnx

#Synced-CL 50925863

diff --git a/Engine/Plugins/Editor/DataValidation/Source/DataValidation/Public/EditorValidatorBase.h b/Engine/Plugins/Editor/DataValidation/Source/DataValidation/Public/EditorValidatorBase.h
--- a/Engine/Plugins/Editor/DataValidation/Source/DataValidation/Public/EditorValidatorBase.h
+++ b/Engine/Plugins/Editor/DataValidation/Source/DataValidation/Public/EditorValidatorBase.h
@@ -34,6 +34,16 @@ public:
 	// If validator doesn't want to validate this asset/in this context, it will return NotValidated which is acceptable.
 	UE_API EDataValidationResult ValidateLoadedAsset(const FAssetData& InAssetData, UObject* InAsset, FDataValidationContext& Context);
 
+	/** Override this to return additional assets that need validation because of other changes in a changelist. e.g. referencers of a specific type may be broken by edits to a dependency. */
+	virtual TArray<FAssetData> GetAssetsToValidateFromChangelist(UDataValidationChangelist* InChangelist, FDataValidationContext& InContext) const
+	{
+		return {};
+	}
+
+	/** Override this to do any state cleanup or add additional validation messages after all per asset validation is complete. */
+	virtual void PostAssetValidation(TArray<TSharedRef<FTokenizedMessage>>& OutMessages) {}
+
+protected:
 	/** 
 	 * Override this to determine whether or not you can use this validator given this context.
 	 * Context can be used to add errors if validation cannot be performed because of some issue 
@@ -48,15 +58,6 @@ public:
 	{
 		return EDataValidationResult::NotValidated;
 	}
-	
-	/** Override this to return additional assets that need validation because of other changes in a changelist. e.g. referencers of a specific type may be broken by edits to a dependency. */
-	virtual TArray<FAssetData> GetAssetsToValidateFromChangelist(UDataValidationChangelist* InChangelist, FDataValidationContext& InContext) const
-	{
-		return {};
-	}
-	
-	/** Override this to do any state cleanup or add additional validation messages after all per asset validation is complete. */
-	virtual void PostAssetValidation(TArray<TSharedRef<FTokenizedMessage>>& OutMessages) {}
 
 	/** Override this to determine whether or not you can use this validator given this usecase */
 	UFUNCTION(BlueprintNativeEvent, Category = "Asset Validation", DisplayName="Can Validate")
@@ -65,14 +66,14 @@ public:
 	/** Override this to determine whether or not you can validate a given asset with this validator */
 	UFUNCTION(BlueprintNativeEvent, Category = "Asset Validation", DisplayName="Can Validate Asset")
 	UE_API bool K2_CanValidateAsset(UObject* InAsset) const;
-	
+	UE_API virtual bool K2_CanValidateAsset_Implementation(UObject* InAsset) const;
+
 	/** Override this in blueprint to validate assets */
 	UFUNCTION(BlueprintNativeEvent, Category = "Asset Validation", DisplayName="Validate Loaded Asset")
 	UE_API EDataValidationResult K2_ValidateLoadedAsset(UObject* InAsset);
-
-	// Declaration for above BlueprintNativeEvent
 	UE_API virtual EDataValidationResult K2_ValidateLoadedAsset_Implementation(UObject* InAsset);
 	
+public:
 	/** Marks the validation as failed and adds an error message. */
 	UFUNCTION(BlueprintCallable, Category = "Asset Validation")
 	UE_API void AssetFails(const UObject* InAsset, const FText& InMessage);
@@ -117,6 +118,7 @@ public:
 	UE_API void AddLegacyValidationErrors(TArray<FText> InErrors);
 	UE_API EDataValidationResult ExtractValidationState(FDataValidationContext& InOutContext) const;
 
+protected:
 	// Legacy entry points from old interface 
 	UE_DEPRECATED("5.4", "CanValidate_Implementation(const EDataValidationUsecase InUsecase) is deprecated, override CanValidateAsset_Implementation(UObject* InObject, FDataValidationContext& InContext) instead")
 	virtual bool CanValidate_Implementation(const EDataValidationUsecase InUsecase) const 
