Removed FName -> FString -> FName conversion looking up WP actor descs

GetActorDescInstanceByPath will call GetActorDescInstanceByName internally using only the leaf-name from the path, and these places already have the leaf-name as a FName

#rnx

#Synced-CL 51168648

diff --git a/Engine/Plugins/Editor/DataValidation/Source/DataValidation/Private/LoadExternalObjectsForValidation.cpp b/Engine/Plugins/Editor/DataValidation/Source/DataValidation/Private/LoadExternalObjectsForValidation.cpp
--- a/Engine/Plugins/Editor/DataValidation/Source/DataValidation/Private/LoadExternalObjectsForValidation.cpp
+++ b/Engine/Plugins/Editor/DataValidation/Source/DataValidation/Private/LoadExternalObjectsForValidation.cpp
@@ -105,7 +105,7 @@ FScopedLoadExternalObjects::FScopedLoadExternalObjects(UObject* InAsset, FDataVa
 			FWorldPartitionLoadingContext::FDeferred LoadingContext;
 			for (const FAssetData& ActorData : ActorsData)
 			{
-				if (const FWorldPartitionActorDescInstance* ActorDesc = ContainersToValidate.GetActorDescInstanceByPath(ActorData.AssetName.ToString()))
+				if (const FWorldPartitionActorDescInstance* ActorDesc = ContainersToValidate.GetActorDescInstanceByName(ActorData.AssetName))
 				{
 					ActorRefs.Add(FWorldPartitionReference(&ContainersToValidate, ActorDesc->GetGuid()));
 				}
diff --git a/Engine/Plugins/Editor/DataValidation/Source/DataValidation/Private/WorldPartitionChangelistValidator.cpp b/Engine/Plugins/Editor/DataValidation/Source/DataValidation/Private/WorldPartitionChangelistValidator.cpp
--- a/Engine/Plugins/Editor/DataValidation/Source/DataValidation/Private/WorldPartitionChangelistValidator.cpp
+++ b/Engine/Plugins/Editor/DataValidation/Source/DataValidation/Private/WorldPartitionChangelistValidator.cpp
@@ -255,7 +255,7 @@ void UWorldPartitionChangelistValidator::ValidateActorsAndDataLayersFromChangeLi
 		for (const FAssetData& ActorData : ActorsData)
 		{
 			// Get the actor descriptor
-			if (const FWorldPartitionActorDescInstance* ActorDescInstance = ContainersToValidate.GetActorDescInstanceByPath(ActorData.AssetName.ToString()))
+			if (const FWorldPartitionActorDescInstance* ActorDescInstance = ContainersToValidate.GetActorDescInstanceByName(ActorData.AssetName))
 			{
 				RelevantActorGuids.Add(ActorDescInstance->GetGuid());
 			}
