Removed option to disable loading external objects during validation, as this is now required by some validators

#rnx

#Synced-CL 51168648

diff --git a/Engine/Plugins/Editor/DataValidation/Source/DataValidation/Private/DataValidationModule.cpp b/Engine/Plugins/Editor/DataValidation/Source/DataValidation/Private/DataValidationModule.cpp
--- a/Engine/Plugins/Editor/DataValidation/Source/DataValidation/Private/DataValidationModule.cpp
+++ b/Engine/Plugins/Editor/DataValidation/Source/DataValidation/Private/DataValidationModule.cpp
@@ -36,12 +36,6 @@
 
 namespace DataValidationOptions
 {
-	static bool bLoadExternalObjectsForCookValidation = true;
-	static FAutoConsoleVariableRef CVarLoadExternalObjectsForCookValidation(
-		TEXT("DataValidation.LoadExternalObjectsForCookValidation"),
-		bLoadExternalObjectsForCookValidation,
-		TEXT("True to load any external objects (eg, actors) prior to running cook-time validation on the asset"));
-
 	static bool bReportCookValidationStats = true;
 	static FAutoConsoleVariableRef CVarReportCookValidationStats(
 		TEXT("DataValidation.ReportCookValidationStats"),
@@ -505,7 +499,7 @@ EDataValidationResult FDataValidationModule::OnValidateSourcePackageDuringCook(U
 						->AddToken(FAssetDataToken::Create(AssetData))
 						->AddToken(FTextToken::Create(LOCTEXT("Data.ValidatingAsset", "Validating asset")));
 
-					UE::DataValidation::FScopedLoadExternalObjects ExternalObjectsLoader(Asset, ValidationContext, DataValidationOptions::bLoadExternalObjectsForCookValidation);
+					UE::DataValidation::FScopedLoadExternalObjects ExternalObjectsLoader(Asset, ValidationContext, true);
 
 					UE::DataValidation::FScopedLogMessageGatherer LogGatherer;
 					EDataValidationResult ValidationResult = EditorValidationSubsystem->ValidateObjectInternal(AssetData, Asset, ValidationContext, DataValidationOptions::bReportCookValidationStats ? &CookValidationStats : nullptr);
