Fixed cooked assets being considered as on-disk dependencies for source control validation

#rnx
#Synced-CL 51168648

diff --git a/Engine/Plugins/Editor/DataValidation/Source/DataValidation/Private/DataValidationChangelist.cpp b/Engine/Plugins/Editor/DataValidation/Source/DataValidation/Private/DataValidationChangelist.cpp
--- a/Engine/Plugins/Editor/DataValidation/Source/DataValidation/Private/DataValidationChangelist.cpp
+++ b/Engine/Plugins/Editor/DataValidation/Source/DataValidation/Private/DataValidationChangelist.cpp
@@ -26,11 +26,30 @@ void UDataValidationChangelist::GatherDependencies(const FName& InPackageName, T
 
 	for (const FName& PackageDependency : Dependencies)
 	{
+		const FString& PackageDependencyStr = PackageDependency.ToString();
+
 		// Exclude script/memory packages
-		if (FPackageName::IsValidLongPackageName(PackageDependency.ToString()))
+		if (!FPackageName::IsValidLongPackageName(PackageDependencyStr))
+		{
+			continue;
+		}
+
+		// Exclude Verse dependencies
+		if (FPackageName::IsVersePackage(PackageDependencyStr))
+		{
+			continue;
+		}
+
+		// Exclude cooked dependencies
 		{
-			OutDependencies.Add({ InPackageName, PackageDependency });
+			USourceControlHelpers::GetAssetDataFromPackage(PackageDependencyStr, Assets);
+			if (Algo::AnyOf(Assets, [](const FAssetData& Asset) { return Asset.HasAnyPackageFlags(PKG_Cooked); }))
+			{
+				continue;
+			}
 		}
+
+		OutDependencies.Add({ InPackageName, PackageDependency });
 	}
 }
 
@@ -134,13 +153,6 @@ EDataValidationResult UDataValidationChangelist::IsDataValid(FDataValidationCont
 			continue;
 		}
 
-		// Some kinds of package dependencies create false positives because they're not checked into source
-		// control as of this writing (1/23/2023). We just skip over them.
-		if (FPackageName::IsVersePackage(ExternalDependencyPair.Value.ToString()))
-		{
-			continue;
-		}
-
 		auto GetReferencerText = [&ExternalDependencyPair]() { return FText::FromString(GetPrettyPackageName(ExternalDependencyPair.Key)); };
 		auto GetReferenceText = [&ExternalDependencyPair]() { return FText::FromString(GetPrettyPackageName(ExternalDependencyPair.Value)); };
 		auto GetExternalPackageFilenameText = [&ExternalPackageFilename]() { return FText::FromString(ExternalPackageFilename); };
