Added external object loading/unloading to the cook-time validation stats reporting

#rnx

#Synced-CL 51332317

diff --git a/Engine/Plugins/Editor/DataValidation/Source/DataValidation/Private/DataValidationModule.cpp b/Engine/Plugins/Editor/DataValidation/Source/DataValidation/Private/DataValidationModule.cpp
--- a/Engine/Plugins/Editor/DataValidation/Source/DataValidation/Private/DataValidationModule.cpp
+++ b/Engine/Plugins/Editor/DataValidation/Source/DataValidation/Private/DataValidationModule.cpp
@@ -18,6 +18,7 @@
 #include "HAL/IConsoleManager.h"
 #include "LoadExternalObjectsForValidation.h"
 #include "Logging/MessageLog.h"
+#include "ProfilingDebugging/ScopedTimers.h"
 #include "UObject/ICookInfo.h"
 
 #include "ContentBrowserMenuContexts.h"
@@ -72,7 +73,14 @@ private:
 	static FText Menu_ValidateDataGetTitle();
 	static void Menu_ValidateData();
 
-	TMap<FTopLevelAssetPath, FValidatorStatistics> CookValidationStats;
+	struct FCookValidationStats
+	{
+		TMap<FTopLevelAssetPath, FValidatorStatistics> ValidatorStats;
+		double ExternalObjectLoadTotalDurationSeconds = 0.0;
+		double ExternalObjectUnloadTotalDurationSeconds = 0.0;
+		int32 ExternalObjectTotalCount = 0;
+	};
+	TOptional<FCookValidationStats> CookValidationStats;
 };
 
 IMPLEMENT_MODULE(FDataValidationModule, DataValidation)
@@ -413,7 +421,14 @@ void FDataValidationModule::OnCookStarted(UE::Cook::ICookInfo& CookInfo)
 {
 	using namespace UE::Cook;
 
-	CookValidationStats.Reset();
+	if (DataValidationOptions::bReportCookValidationStats)
+	{
+		CookValidationStats.Emplace();
+	}
+	else
+	{
+		CookValidationStats.Reset();
+	}
 
 	if (EnumHasAnyFlags(CookInfo.GetCookValidationOptions(), ECookValidationOptions::RunAssetValidation))
 	{
@@ -431,7 +446,10 @@ void FDataValidationModule::OnCookStarted(UE::Cook::ICookInfo& CookInfo)
 				return true;
 			});
 
-			CookValidationStats.Reserve(NumEnabledValidators);
+			if (CookValidationStats)
+			{
+				CookValidationStats->ValidatorStats.Reserve(NumEnabledValidators);
+			}
 		}
 	}
 }
@@ -440,18 +458,19 @@ void FDataValidationModule::OnCookFinished(UE::Cook::ICookInfo& CookInfo)
 {
 	using namespace UE::Cook;
 
-	if (DataValidationOptions::bReportCookValidationStats && EnumHasAnyFlags(CookInfo.GetCookValidationOptions(), ECookValidationOptions::RunAssetValidation))
+	if (CookValidationStats && EnumHasAnyFlags(CookInfo.GetCookValidationOptions(), ECookValidationOptions::RunAssetValidation))
 	{
 		// Report the validators, slowest to fastest
-		CookValidationStats.ValueSort([](const FValidatorStatistics& One, const FValidatorStatistics& Two)
+		CookValidationStats->ValidatorStats.ValueSort(
+			[](const FValidatorStatistics& One, const FValidatorStatistics& Two)
 			{
 				return One.TotalDurationSeconds > Two.TotalDurationSeconds;
 			});
 
-		UE_LOG(LogContentValidation, Display, TEXT("Validation stats:"));
-		for (const TTuple<FTopLevelAssetPath, FValidatorStatistics>& CookValidationStatsPair : CookValidationStats)
+		UE_LOG(LogContentValidation, Display, TEXT("Validation stats: (%d external objects took %f seconds to load and %f seconds to unload)"), CookValidationStats->ExternalObjectTotalCount, CookValidationStats->ExternalObjectLoadTotalDurationSeconds, CookValidationStats->ExternalObjectUnloadTotalDurationSeconds);
+		for (const TTuple<FTopLevelAssetPath, FValidatorStatistics>& ValidatorStatsPair : CookValidationStats->ValidatorStats)
 		{
-			UE_LOG(LogContentValidation, Display, TEXT("\t%s validated %d assets in %f seconds"), *CookValidationStatsPair.Key.ToString(), CookValidationStatsPair.Value.AssetsValidated, CookValidationStatsPair.Value.TotalDurationSeconds);
+			UE_LOG(LogContentValidation, Display, TEXT("\t%s validated %d assets in %f seconds"), *ValidatorStatsPair.Key.ToString(), ValidatorStatsPair.Value.AssetsValidated, ValidatorStatsPair.Value.TotalDurationSeconds);
 		}
 	}
 }
@@ -459,7 +478,23 @@ void FDataValidationModule::OnCookFinished(UE::Cook::ICookInfo& CookInfo)
 #if ENABLE_COOK_STATS
 void FDataValidationModule::LogCookStats(FCookStatsManager::AddStatFuncRef AddStat)
 {
-	UE::DataValidation::ValidationStatsToAnalyticEventData(CookValidationStats, AddStat);
+	if (CookValidationStats)
+	{
+		UE::DataValidation::ValidationStatsToAnalyticEventData(CookValidationStats->ValidatorStats, AddStat);
+
+		if (CookValidationStats->ExternalObjectTotalCount > 0)
+		{
+			AddStat(
+				TEXT("Validation_ExternalObjectLoad"),
+				FCookStatsManager::CreateKeyValueArray(TEXT("TotalDurationSecs"), CookValidationStats->ExternalObjectLoadTotalDurationSeconds, TEXT("AssetsValidated"), CookValidationStats->ExternalObjectTotalCount)
+				);
+
+			AddStat(
+				TEXT("Validation_ExternalObjectUnload"),
+				FCookStatsManager::CreateKeyValueArray(TEXT("TotalDurationSecs"), CookValidationStats->ExternalObjectUnloadTotalDurationSeconds, TEXT("AssetsValidated"), CookValidationStats->ExternalObjectTotalCount)
+				);
+		}
+	}
 }
 #endif
 
@@ -499,10 +534,27 @@ EDataValidationResult FDataValidationModule::OnValidateSourcePackageDuringCook(U
 						->AddToken(FAssetDataToken::Create(AssetData))
 						->AddToken(FTextToken::Create(LOCTEXT("Data.ValidatingAsset", "Validating asset")));
 
-					UE::DataValidation::FScopedLoadExternalObjects ExternalObjectsLoader(Asset, ValidationContext, true);
+					TOptional<UE::DataValidation::FScopedLoadExternalObjects> ExternalObjectsLoader;
+					if (!ValidationContext.GetAssociatedExternalObjects().IsEmpty())
+					{
+						if (CookValidationStats)
+						{
+							double LoadDuration = 0.0;
+							{
+								FScopedDurationTimer ValidationTimer(LoadDuration);
+								ExternalObjectsLoader.Emplace(Asset, ValidationContext, true);
+							}
+							CookValidationStats->ExternalObjectLoadTotalDurationSeconds += LoadDuration;
+							CookValidationStats->ExternalObjectTotalCount += ValidationContext.GetAssociatedExternalObjects().Num();
+						}
+						else
+						{
+							ExternalObjectsLoader.Emplace(Asset, ValidationContext, true);
+						}
+					}
 
 					UE::DataValidation::FScopedLogMessageGatherer LogGatherer;
-					EDataValidationResult ValidationResult = EditorValidationSubsystem->ValidateObjectInternal(AssetData, Asset, ValidationContext, DataValidationOptions::bReportCookValidationStats ? &CookValidationStats : nullptr);
+					EDataValidationResult ValidationResult = EditorValidationSubsystem->ValidateObjectInternal(AssetData, Asset, ValidationContext, CookValidationStats ? &CookValidationStats->ValidatorStats : nullptr);
 
 					TArray<FString> LogWarnings;
 					TArray<FString> LogErrors;
@@ -530,6 +582,23 @@ EDataValidationResult FDataValidationModule::OnValidateSourcePackageDuringCook(U
 
 					UE::DataValidation::AddAssetValidationMessages(AssetData, DataValidationLog, ValidationContext);
 					DataValidationLog.Flush();
+
+					if (ExternalObjectsLoader)
+					{
+						if (CookValidationStats)
+						{
+							double UnloadDuration = 0.0;
+							{
+								FScopedDurationTimer ValidationTimer(UnloadDuration);
+								ExternalObjectsLoader.Reset();
+							}
+							CookValidationStats->ExternalObjectUnloadTotalDurationSeconds += UnloadDuration;
+						}
+						else
+						{
+							ExternalObjectsLoader.Reset();
+						}
+					}
 				}
 			}
 
