We only validate on-disk assets when cooking, so can skip the in-memory query when finding FAssetData

#rnx

#Synced-CL 51332317

diff --git a/Engine/Plugins/Editor/DataValidation/Source/DataValidation/Private/DataValidationModule.cpp b/Engine/Plugins/Editor/DataValidation/Source/DataValidation/Private/DataValidationModule.cpp
--- a/Engine/Plugins/Editor/DataValidation/Source/DataValidation/Private/DataValidationModule.cpp
+++ b/Engine/Plugins/Editor/DataValidation/Source/DataValidation/Private/DataValidationModule.cpp
@@ -505,7 +505,7 @@ EDataValidationResult FDataValidationModule::OnValidateSourcePackageDuringCook(U
 	if (UEditorValidatorSubsystem* EditorValidationSubsystem = GEditor->GetEditorSubsystem<UEditorValidatorSubsystem>())
 	{
 		TArray<FAssetData> AssetList;
-		IAssetRegistry::GetChecked().GetAssetsByPackageName(Package->GetFName(), AssetList);
+		IAssetRegistry::GetChecked().GetAssetsByPackageName(Package->GetFName(), AssetList, /*bIncludeOnlyOnDiskAssets*/true);
 
 		{
 			FValidateAssetsSettings Settings;
