Remove double validation if packages are saved immediately prior to being validated

Note: This only helps if there's no Tick between the save and the validate, otherwise the latent save validation queue will be flushed

#rnx
#Synced-CL 51753475

diff --git a/Engine/Plugins/Editor/DataValidation/Source/DataValidation/Private/EditorValidatorSubsystem.cpp b/Engine/Plugins/Editor/DataValidation/Source/DataValidation/Private/EditorValidatorSubsystem.cpp
--- a/Engine/Plugins/Editor/DataValidation/Source/DataValidation/Private/EditorValidatorSubsystem.cpp
+++ b/Engine/Plugins/Editor/DataValidation/Source/DataValidation/Private/EditorValidatorSubsystem.cpp
@@ -676,6 +676,16 @@ EDataValidationResult UEditorValidatorSubsystem::ValidateAssetsInternal(
 			ValidationExternalObjects = *ValidationExternalObjectsPtr;
 		}
 
+		// This asset is about to be validated, so remove it from the latent save validation list
+		if (SavedPackagesToValidate.Num() > 0)
+		{
+			SavedPackagesToValidate.Remove(Data.PackageName);
+			for (const FAssetData& ExternalObject : ValidationExternalObjects)
+			{
+				SavedPackagesToValidate.Remove(ExternalObject.PackageName);
+			}
+		}
+
 		FDataValidationContext ValidationContext(false, InSettings.ValidationUsecase, ValidationExternalObjects, InSettings.MessageLogName);
 		EDataValidationResult AssetResult = EDataValidationResult::NotValidated;
 		if (!LoadedAsset)
@@ -951,7 +961,7 @@ void UEditorValidatorSubsystem::ValidateSavedPackage(FName PackageName, bool bPr
 		GEditor->GetTimerManager()->SetTimerForNextTick(this, &UEditorValidatorSubsystem::ValidateAllSavedPackages);
 	}
 
-	SavedPackagesToValidate.AddUnique(PackageName);
+	SavedPackagesToValidate.Add(PackageName);
 }
 
 bool UEditorValidatorSubsystem::IsPathExcludedFromValidation(FStringView Path) const
@@ -983,16 +993,16 @@ void UEditorValidatorSubsystem::ValidateAllSavedPackages()
 	}
 	// We need to query the in-memory data as the disk cache may not be accurate
 	FARFilter Filter;
-	Filter.PackageNames = SavedPackagesToValidate;
+	Filter.PackageNames = SavedPackagesToValidate.Array();
 	Filter.bIncludeOnlyOnDiskAssets = false;
 
 	TArray<FAssetData> Assets;
 	AssetRegistry.GetAssets(Filter, Assets);
 
+	SavedPackagesToValidate.Empty();
+
 	bool bProceduralSave = false; // The optional suppression for ProceduralSaves was checked before adding to SavedPackagesToValidate
 	ValidateOnSave(MoveTemp(Assets), bProceduralSave);
-
-	SavedPackagesToValidate.Empty();
 }
 
 void UEditorValidatorSubsystem::ValidateChangelistPreSubmit(
diff --git a/Engine/Plugins/Editor/DataValidation/Source/DataValidation/Public/EditorValidatorSubsystem.h b/Engine/Plugins/Editor/DataValidation/Source/DataValidation/Public/EditorValidatorSubsystem.h
--- a/Engine/Plugins/Editor/DataValidation/Source/DataValidation/Public/EditorValidatorSubsystem.h
+++ b/Engine/Plugins/Editor/DataValidation/Source/DataValidation/Public/EditorValidatorSubsystem.h
@@ -537,7 +537,7 @@ protected:
 	bool bValidateOnSave;
 
 	/** List of saved package names to validate next frame */
-	TArray<FName> SavedPackagesToValidate;
+	mutable TSet<FName> SavedPackagesToValidate;
 
 	/**
 	 * Active validator instances, mapped from their class.
